Fix for uninitialized struct members in QAGame

diff --git a/Engine/Plugins/Developer/Concert/DisasterRecoveryClient/Source/DisasterRecoveryClient/Private/DisasterRecoverySessionInfo.h b/Engine/Plugins/Developer/Concert/DisasterRecoveryClient/Source/DisasterRecoveryClient/Private/DisasterRecoverySessionInfo.h
--- a/Engine/Plugins/Developer/Concert/DisasterRecoveryClient/Source/DisasterRecoveryClient/Private/DisasterRecoverySessionInfo.h
+++ b/Engine/Plugins/Developer/Concert/DisasterRecoveryClient/Source/DisasterRecoveryClient/Private/DisasterRecoverySessionInfo.h
@@ -103,7 +103,7 @@ struct FDisasterRecoveryClientInfo
 
 	/** The client process ID. */
 	UPROPERTY()
-	uint32 ClientProcessId;
+	uint32 ClientProcessId = 0;
 
 	/** The client app ID. */
 	UPROPERTY()
diff --git a/Engine/Plugins/Experimental/Animation/MotionTrajectory/Source/MotionTrajectory/Private/MotionTrajectoryLibrary.cpp b/Engine/Plugins/Experimental/Animation/MotionTrajectory/Source/MotionTrajectory/Private/MotionTrajectoryLibrary.cpp
--- a/Engine/Plugins/Experimental/Animation/MotionTrajectory/Source/MotionTrajectory/Private/MotionTrajectoryLibrary.cpp
+++ b/Engine/Plugins/Experimental/Animation/MotionTrajectory/Source/MotionTrajectory/Private/MotionTrajectoryLibrary.cpp
@@ -1,6 +1,7 @@
 // Copyright Epic Games, Inc. All Rights Reserved.
 
 #include "MotionTrajectoryLibrary.h"
+#include "GameFramework/Actor.h"
 #include "Algo/Find.h"
 
 static void FlattenTrajectoryPosition(FTrajectorySample& Sample, float& AccumulatedDeltaDistance, bool PreserveSpeed)
diff --git a/Engine/Plugins/Experimental/Animation/MotionTrajectory/Source/MotionTrajectory/Public/MotionTrajectoryLibrary.h b/Engine/Plugins/Experimental/Animation/MotionTrajectory/Source/MotionTrajectory/Public/MotionTrajectoryLibrary.h
--- a/Engine/Plugins/Experimental/Animation/MotionTrajectory/Source/MotionTrajectory/Public/MotionTrajectoryLibrary.h
+++ b/Engine/Plugins/Experimental/Animation/MotionTrajectory/Source/MotionTrajectory/Public/MotionTrajectoryLibrary.h
@@ -15,11 +15,11 @@ struct MOTIONTRAJECTORY_API FTrajectoryDirectionClamp
 
 	// Target clamped direction for an incoming source direction
 	UPROPERTY(EditAnywhere, BlueprintReadWrite, Category="Settings")
-	FVector Direction;
+	FVector Direction = FVector::ZeroVector;
 
 	// Target clamped directions will be applied to any source direction within this angle boundary
 	UPROPERTY(EditAnywhere, BlueprintReadWrite, Category="Settings", meta=(ClampMin="1.0", ClampMax="180.0"))
-	float AngleTresholdDegrees;
+	float AngleTresholdDegrees = 0.f;
 };
 
 UCLASS(Category="Motion Trajectory")
