Added angular information to trajectory samples

diff --git a/Engine/Plugins/Experimental/Animation/MotionTrajectory/Source/MotionTrajectory/Private/MotionTrajectory.cpp b/Engine/Plugins/Experimental/Animation/MotionTrajectory/Source/MotionTrajectory/Private/MotionTrajectory.cpp
--- a/Engine/Plugins/Experimental/Animation/MotionTrajectory/Source/MotionTrajectory/Private/MotionTrajectory.cpp
+++ b/Engine/Plugins/Experimental/Animation/MotionTrajectory/Source/MotionTrajectory/Private/MotionTrajectory.cpp
@@ -9,6 +9,8 @@
 
 #define LOCTEXT_NAMESPACE "MotionTrajectory"
 
+DEFINE_LOG_CATEGORY(LogMotionTrajectory);
+
 void FMotionTrajectoryModule::StartupModule()
 {
 }
@@ -24,18 +26,34 @@ UMotionTrajectoryComponent::UMotionTrajectoryComponent(const FObjectInitializer&
 	PrimaryComponentTick.TickGroup = TG_PrePhysics;
 }
 
-FTrajectorySample UMotionTrajectoryComponent::GetPresentTrajectory() const
+FTrajectorySample UMotionTrajectoryComponent::CalcWorldSpacePresentTrajectorySample(float DeltaTime) const
 {
 	const APawn* Pawn = TryGetOwnerPawn();
 	checkf(false, TEXT("UMotionTrajectoryComponent::GetPresentTrajectory for Pawn: %s requires implementation."), Pawn ? *(Pawn->GetHumanReadableName()) : TEXT("NULL"));
 	return FTrajectorySample();
 }
 
-FTransform UMotionTrajectoryComponent::GetPresentWorldTransform() const
+bool UMotionTrajectoryComponent::HasTrajectoryTickedThisFrame() const
 {
-	const APawn* Pawn = TryGetOwnerPawn();
-	checkf(false, TEXT("UMotionTrajectoryComponent::GetPresentWorldTransform for Pawn: %s requires implementation."), Pawn ? *(Pawn->GetHumanReadableName()) : TEXT("NULL"));
-	return FTransform::Identity;
+	return GFrameCounter == LastTrajectoryTickFrame;
+}
+
+void UMotionTrajectoryComponent::TickTrajectory(float DeltaTime)
+{
+	LastTrajectoryTickFrame = static_cast<uint32>(GFrameCounter);
+
+	// Compute the instantaneous/present trajectory sample and guarantee that is zeroed on all accumulated domains
+	PresentTrajectorySampleWS = CalcWorldSpacePresentTrajectorySample(DeltaTime);
+	PresentTrajectorySampleWS.AccumulatedDistance = 0.f;
+	PresentTrajectorySampleWS.AccumulatedSeconds = 0.f;
+
+	// convert to local space
+	PresentTrajectorySampleLS = PresentTrajectorySampleWS;
+	PresentTrajectorySampleLS.PrependOffset(PresentTrajectorySampleWS.Transform.Inverse(), 0.0f);
+
+
+	// Tick the historical sample retention/decay algorithm by one iteration
+	TickHistoryEvictionPolicy();
 }
 
 const APawn* UMotionTrajectoryComponent::TryGetOwnerPawn() const
@@ -60,7 +78,7 @@ FTrajectorySampleRange UMotionTrajectoryComponent::CombineHistoryPresentPredicti
 	// Ex: [Past + Present + Future], with negative values representing the past, a zeroed present, and positive values representing the future
 	Trajectory.Samples.Reserve(TotalSampleCount);
 	Trajectory.Samples.Append(History.Samples);
-	Trajectory.Samples.Add(PresentTrajectory);
+	Trajectory.Samples.Add(PresentTrajectorySampleLS);
 	Trajectory.Samples.Append(Prediction.Samples);
 	Trajectory.SampleRate = SampleRate;
 #if WITH_EDITORONLY_DATA
@@ -73,7 +91,7 @@ void UMotionTrajectoryComponent::TickHistoryEvictionPolicy()
 {
 	// World space transform and time are used for determining local-relative historical samples
 	const UWorld* World = GetWorld();
-	const FTransform WorldTransform = GetPresentWorldTransform();
+	const FTransform WorldTransform = PresentTrajectorySampleWS.Transform;
 	const float WorldGameTime = UKismetSystemLibrary::GetGameTimeInSeconds(World);
 
 	// Skip on the very first sample
@@ -86,9 +104,7 @@ void UMotionTrajectoryComponent::TickHistoryEvictionPolicy()
 
 		for (auto& Sample : SampleHistory)
 		{
-			Sample.AccumulatedSeconds -= DeltaSeconds;
-			Sample.AccumulatedDistance -= DeltaDistance;
-			Sample.Position = DeltaTransform.TransformPosition(Sample.Position);
+			Sample.PrependOffset(DeltaTransform, -DeltaSeconds);
 		}
 
 		const FTrajectorySample FirstSample = SampleHistory.First();
@@ -120,7 +136,7 @@ void UMotionTrajectoryComponent::TickHistoryEvictionPolicy()
 		SampleHistory.RemoveAll([&](const FTrajectorySample& Sample)
 			{
 				// Remove superfluous zero motion samples
-				if (Sample.IsZeroSample() && PresentTrajectory.IsZeroSample())
+				if (Sample.IsZeroSample() && PresentTrajectorySampleLS.IsZeroSample())
 				{
 					return true;
 				}
@@ -151,7 +167,7 @@ void UMotionTrajectoryComponent::TickHistoryEvictionPolicy()
 	// An opportunity to insert more samples will arise as history decays
 	if (SampleHistory.Num() < MaxSamples)
 	{
-		SampleHistory.Emplace(PresentTrajectory);
+		SampleHistory.Emplace(PresentTrajectorySampleLS);
 	}
 	
 	// Cache the present world space sample as the next frame's reference point for delta computation
@@ -174,17 +190,33 @@ void UMotionTrajectoryComponent::OnComponentCreated()
 	Super::OnComponentCreated();
 }
 
-void UMotionTrajectoryComponent::TickComponent(float DeltaTime, enum ELevelTick TickType, FActorComponentTickFunction* ThisTickFunction)
+void UMotionTrajectoryComponent::BeginPlay() 
 {
-	// Compute the instantaneous/present trajectory sample and guarantee that is zeroed on all accumulated domains
-	PresentTrajectory = GetPresentTrajectory();
-	PresentTrajectory.AccumulatedDistance = 0.f;
-	PresentTrajectory.AccumulatedSeconds = 0.f;
+	Super::BeginPlay();
 
-	// Tick the historical sample retention/decay algorithm by one iteration
-	TickHistoryEvictionPolicy();
+	FTransform PresentTransform = FTransform::Identity;
+
+	if (const APawn* Pawn = TryGetOwnerPawn())
+	{
+		PresentTransform = Pawn->GetActorTransform();
+	}
+
+	PresentTrajectorySampleWS = FTrajectorySample();
+	PresentTrajectorySampleWS.Transform = PresentTransform;
+
+	PresentTrajectorySampleLS = FTrajectorySample();
 
+	LastTrajectoryTickFrame = 0u;
+}
+
+void UMotionTrajectoryComponent::TickComponent(float DeltaTime, enum ELevelTick TickType, FActorComponentTickFunction* ThisTickFunction)
+{
 	Super::TickComponent(DeltaTime, TickType, ThisTickFunction);
+
+	if (!HasTrajectoryTickedThisFrame())
+	{
+		TickTrajectory(DeltaTime);
+	}
 }
 
 FTrajectorySampleRange UMotionTrajectoryComponent::GetTrajectory() const
diff --git a/Engine/Plugins/Experimental/Animation/MotionTrajectory/Source/MotionTrajectory/Private/MotionTrajectoryCharacterMovement.cpp b/Engine/Plugins/Experimental/Animation/MotionTrajectory/Source/MotionTrajectory/Private/MotionTrajectoryCharacterMovement.cpp
--- a/Engine/Plugins/Experimental/Animation/MotionTrajectory/Source/MotionTrajectory/Private/MotionTrajectoryCharacterMovement.cpp
+++ b/Engine/Plugins/Experimental/Animation/MotionTrajectory/Source/MotionTrajectory/Private/MotionTrajectoryCharacterMovement.cpp
@@ -5,6 +5,7 @@
 #include "Animation/AnimInstance.h"
 #include "GameFramework/CharacterMovementComponent.h"
 #include "GameFramework/Pawn.h"
+#include "GameFramework/Character.h"
 #include "KismetAnimationLibrary.h"
 
 // ----------- BEGIN Derived from FCharacterMovementComponentAsyncInput::CalcVelocity() ----------- //
@@ -13,11 +14,13 @@ static void StepPrediction(float IntegrationDelta
 	, const UCharacterMovementComponent* MovementComponent
 	, FTrajectorySample& Sample)
 {
+	const FTransform InitialSampleTransform = Sample.Transform;
+
 	if (MovementComponent->GetCurrentAcceleration().IsZero())
 	{
 		float ActualBrakingFriction = (MovementComponent->bUseSeparateBrakingFriction ? MovementComponent->BrakingFriction : MovementComponent->GroundFriction);
 
-		if (!Sample.LocalLinearVelocity.IsZero())
+		if (!Sample.LinearVelocity.IsZero())
 		{
 			const float FrictionFactor = FMath::Max(0.f, MovementComponent->BrakingFrictionFactor);
 			ActualBrakingFriction = FMath::Max(0.f, ActualBrakingFriction * FrictionFactor);
@@ -34,43 +37,127 @@ static void StepPrediction(float IntegrationDelta
 			float RemainingTime = IntegrationDelta;
 			const float MaxTimeStep = FMath::Clamp(MovementComponent->BrakingSubStepTime, 1.f / 75.f, 1.f / 20.f);
 
-			const FVector PrevLinearVelocity = Sample.LocalLinearVelocity;
-			const FVector RevAccel = (bZeroBraking ? FVector::ZeroVector : (-BrakingDeceleration * Sample.LocalLinearVelocity.GetSafeNormal()));
+			const FVector PrevLinearVelocity = Sample.LinearVelocity;
+			const FVector RevAccel = (bZeroBraking ? FVector::ZeroVector : (-BrakingDeceleration * Sample.LinearVelocity.GetSafeNormal()));
 
 			while (RemainingTime >= UCharacterMovementComponent::MIN_TICK_TIME)
 			{
 				const float dt = ((RemainingTime > MaxTimeStep && !bZeroFriction) ? FMath::Min(MaxTimeStep, RemainingTime * 0.5f) : RemainingTime);
 				RemainingTime -= dt;
 
-				Sample.LocalLinearVelocity = Sample.LocalLinearVelocity + ((-ActualBrakingFriction) * Sample.LocalLinearVelocity + RevAccel) * dt;
+				Sample.LinearVelocity = Sample.LinearVelocity + ((-ActualBrakingFriction) * Sample.LinearVelocity + RevAccel) * dt;
 
-				if ((Sample.LocalLinearVelocity | PrevLinearVelocity) <= 0.f)
+				if ((Sample.LinearVelocity | PrevLinearVelocity) <= 0.f)
 				{
-					Sample.LocalLinearVelocity = FVector::ZeroVector;
+					Sample.LinearVelocity = FVector::ZeroVector;
 					return;
 				}
 			}
 
 			// Clamp to zero if nearly zero, or if below min threshold and braking
-			const float VSizeSq = Sample.LocalLinearVelocity.SizeSquared();
+			const float VSizeSq = Sample.LinearVelocity.SizeSquared();
 			if (VSizeSq <= KINDA_SMALL_NUMBER || (!bZeroBraking && VSizeSq <= FMath::Square(UCharacterMovementComponent::BRAKE_TO_STOP_VELOCITY)))
 			{
-				Sample.LocalLinearVelocity = FVector::ZeroVector;
+				Sample.LinearVelocity = FVector::ZeroVector;
 			}
 		}
 	}
 	else
 	{
-		const FVector AccelDir = Sample.LocalLinearAcceleration.GetSafeNormal();
-		const float VelSize = Sample.LocalLinearVelocity.Size();
-		Sample.LocalLinearVelocity = Sample.LocalLinearVelocity - (Sample.LocalLinearVelocity - AccelDir * VelSize) * FMath::Min(IntegrationDelta * MovementComponent->GroundFriction, 1.f);
+		const FVector AccelDir = Sample.LinearAcceleration.GetSafeNormal();
+		const float VelSize = Sample.LinearVelocity.Size();
+		Sample.LinearVelocity = Sample.LinearVelocity - (Sample.LinearVelocity - AccelDir * VelSize) * FMath::Min(IntegrationDelta * MovementComponent->GroundFriction, 1.f);
 
 		const float MaxInputSpeed = FMath::Max(MovementComponent->GetMaxSpeed() * MovementComponent->GetAnalogInputModifier(), MovementComponent->GetMinAnalogSpeed());
-		Sample.LocalLinearVelocity += Sample.LocalLinearAcceleration * IntegrationDelta;
-		Sample.LocalLinearVelocity = Sample.LocalLinearVelocity.GetClampedToMaxSize(MaxInputSpeed);
+		Sample.LinearVelocity += Sample.LinearAcceleration * IntegrationDelta;
+		Sample.LinearVelocity = Sample.LinearVelocity.GetClampedToMaxSize(MaxInputSpeed);
 	}
 
-	Sample.Position += Sample.LocalLinearVelocity * IntegrationDelta;
+	const FTransform InitialTransformWS = InitialSampleTransform * MovementComponent->GetOwner()->GetActorTransform();
+
+	if (MovementComponent->bOrientRotationToMovement)
+	{
+		FRotator CurrentRotation = InitialTransformWS.GetRotation().Rotator();
+		CurrentRotation.DiagnosticCheckNaN(TEXT("CharacterMovementComponent::PhysicsRotation(): CurrentRotation"));
+
+		FRotator DeltaRot = MovementComponent->GetDeltaRotation(IntegrationDelta);
+		DeltaRot.DiagnosticCheckNaN(TEXT("CharacterMovementComponent::PhysicsRotation(): GetDeltaRotation"));
+
+		FRotator DesiredRotation = 
+			MovementComponent->ComputeOrientToMovementRotation(CurrentRotation, IntegrationDelta, DeltaRot);
+
+		// TODO: leaving this code commented here as a reference for later combining both modes so we have one 
+		// component of rotation targeting a fixed direction (above) and the below being extrapolated
+		//FRotator DesiredRotation = CurrentRotation;
+		//if (MovementComponent->bOrientRotationToMovement)
+		//{
+		//	DesiredRotation = MovementComponent->ComputeOrientToMovementRotation(CurrentRotation, IntegrationDelta, DeltaRot);
+		//}
+		//else if (MovementComponent->GetCharacterOwner()->Controller && MovementComponent->bUseControllerDesiredRotation)
+		//{
+		//	DesiredRotation = MovementComponent->GetCharacterOwner()->Controller->GetDesiredRotation();
+		//}
+
+		if (MovementComponent->ShouldRemainVertical())
+		{
+			DesiredRotation.Pitch = 0.f;
+			DesiredRotation.Yaw = FRotator::NormalizeAxis(DesiredRotation.Yaw);
+			DesiredRotation.Roll = 0.f;
+		}
+		else
+		{
+			DesiredRotation.Normalize();
+		}
+
+		// Accumulate a desired new rotation.
+		const float AngleTolerance = 1e-3f;
+
+		if (!CurrentRotation.Equals(DesiredRotation, AngleTolerance))
+		{
+			// PITCH
+			if (!FMath::IsNearlyEqual(CurrentRotation.Pitch, DesiredRotation.Pitch, AngleTolerance))
+			{
+				DesiredRotation.Pitch = FMath::FixedTurn(CurrentRotation.Pitch, DesiredRotation.Pitch, DeltaRot.Pitch);
+			}
+
+			// YAW
+			if (!FMath::IsNearlyEqual(CurrentRotation.Yaw, DesiredRotation.Yaw, AngleTolerance))
+			{
+				DesiredRotation.Yaw = FMath::FixedTurn(CurrentRotation.Yaw, DesiredRotation.Yaw, DeltaRot.Yaw);
+			}
+
+			// ROLL
+			if (!FMath::IsNearlyEqual(CurrentRotation.Roll, DesiredRotation.Roll, AngleTolerance))
+			{
+				DesiredRotation.Roll = FMath::FixedTurn(CurrentRotation.Roll, DesiredRotation.Roll, DeltaRot.Roll);
+			}
+
+			// Set the new rotation.
+			DesiredRotation.DiagnosticCheckNaN(TEXT("CharacterMovementComponent::PhysicsRotation(): DesiredRotation"));
+			const FQuat DesiredRotationQuat = DesiredRotation.Quaternion();
+			const FQuat DesiredRotationQuatAS = MovementComponent->GetOwner()->GetActorQuat().Inverse() * DesiredRotationQuat;
+			Sample.Transform.SetRotation(DesiredRotationQuatAS);
+		}
+
+		const FQuat DeltaRotation = InitialSampleTransform.GetRotation().Inverse() * Sample.Transform.GetRotation();
+
+		// since we're not extrapolating the rotation on this mode, the linear component's don't rotate
+		//Sample.LinearVelocity = DeltaRotation.RotateVector(Sample.LinearVelocity);
+		//Sample.LinearAcceleration = DeltaRotation.UnrotateVector(Sample.LinearAcceleration);
+	}
+	else
+	{
+		const FQuat DeltaRotation = FQuat(Sample.AngularVelocityAxis, Sample.AngularSpeed * IntegrationDelta);
+		const FQuat SampleRotation = Sample.Transform.GetRotation() * DeltaRotation;
+		Sample.Transform.SetRotation(SampleRotation);
+
+		Sample.LinearVelocity = DeltaRotation.RotateVector(Sample.LinearVelocity);
+		Sample.LinearAcceleration = DeltaRotation.RotateVector(Sample.LinearAcceleration);
+	}
+
+	
+	const FVector Translation = Sample.LinearVelocity * IntegrationDelta;
+	Sample.Transform.AddToTranslation(Translation);
 }
 
 static void PredictTrajectory(const UCharacterMovementComponent* MovementComponent
@@ -98,7 +185,8 @@ static void PredictTrajectory(const UCharacterMovementComponent* MovementCompone
 			PreviousSample = Sample;
 			StepPrediction(IntegrationDelta, MovementComponent, Sample);
 
-			AccumulatedDistance += FVector::Distance(PreviousSample.Position, Sample.Position);
+			AccumulatedDistance += 
+				FVector::Distance(PreviousSample.Transform.GetLocation(), Sample.Transform.GetLocation());
 			Sample.AccumulatedDistance = AccumulatedDistance;
 			AccumulatedSeconds += IntegrationDelta;
 			Sample.AccumulatedSeconds = AccumulatedSeconds;
@@ -132,7 +220,7 @@ static void PredictTrajectory(const UCharacterMovementComponent* MovementCompone
 
 // ------------ END Derived from FCharacterMovementComponentAsyncInput::CalcVelocity() ------------ //
 
-FTrajectorySample UCharacterMovementTrajectoryComponent::GetPresentTrajectory() const
+FTrajectorySample UCharacterMovementTrajectoryComponent::CalcWorldSpacePresentTrajectorySample(float DeltaTime) const
 {
 	FTrajectorySample ReturnValue;
 	
@@ -145,28 +233,32 @@ FTrajectorySample UCharacterMovementTrajectoryComponent::GetPresentTrajectory()
 	const UCharacterMovementComponent* MovementComponent = Cast<UCharacterMovementComponent>(Pawn->GetMovementComponent());
 	if (ensure(MovementComponent) && MovementComponent->MovementMode == EMovementMode::MOVE_Walking)
 	{
-		const USkeletalMeshComponent* SkeletalMeshComponent = Pawn->FindComponentByClass<USkeletalMeshComponent>();
-		FTransform ComponentWorldTransform = SkeletalMeshComponent ? SkeletalMeshComponent->GetComponentTransform() : FTransform();
-
-		ReturnValue.LocalLinearVelocity = ComponentWorldTransform.InverseTransformVectorNoScale(MovementComponent->Velocity);
-		ReturnValue.LocalLinearAcceleration = ComponentWorldTransform.InverseTransformVectorNoScale(MovementComponent->GetCurrentAcceleration());
-	}
-	return ReturnValue;
-}
+		FTransform ComponentWorldTransform = Pawn->GetActorTransform();
+		FVector MovementComponentPosition = MovementComponent->GetLastUpdateLocation();
 
-FTransform UCharacterMovementTrajectoryComponent::GetPresentWorldTransform() const
-{
-	FTransform ReturnValue = FTransform::Identity;
+		ReturnValue.Transform = ComponentWorldTransform;
+		ReturnValue.LinearVelocity = MovementComponent->Velocity;
+		ReturnValue.LinearAcceleration = MovementComponent->GetCurrentAcceleration();
 
-	const APawn* Pawn = TryGetOwnerPawn();
-	if (!Pawn)
-	{
-		return ReturnValue;
+		if (DeltaTime > SMALL_NUMBER)
+		{
+			const FQuat DeltaRotation = 
+				PresentTrajectorySampleWS.Transform.GetRotation().Inverse() * ComponentWorldTransform.GetRotation();
+			FVector DeltaRotationAxis;
+			float DeltaRotationAngle;
+			DeltaRotation.ToAxisAndAngle(DeltaRotationAxis, DeltaRotationAngle);
+			const float AngularSpeed = DeltaRotationAngle / DeltaTime;
+
+			ReturnValue.AngularVelocityAxis = DeltaRotationAxis;
+			ReturnValue.AngularSpeed = AngularSpeed;
+		}
+		else
+		{
+			ReturnValue.AngularVelocityAxis = PresentTrajectorySampleWS.AngularVelocityAxis;
+			ReturnValue.AngularSpeed = PresentTrajectorySampleWS.AngularSpeed;
+		}
 	}
 
-	const USkeletalMeshComponent* SkeletalMeshComponent = Pawn->FindComponentByClass<USkeletalMeshComponent>();
-	ReturnValue = SkeletalMeshComponent ? SkeletalMeshComponent->GetComponentTransform() : Pawn->GetActorTransform();
-
 	return ReturnValue;
 }
 
@@ -181,6 +273,13 @@ void UCharacterMovementTrajectoryComponent::InitializeComponent()
 	UCharacterMovementComponent* CharacterMovementComponent = GetOwner()->FindComponentByClass<UCharacterMovementComponent>();
 	AddTickPrerequisiteComponent(CharacterMovementComponent);
 
+	if (ACharacter* Character = CharacterMovementComponent->GetCharacterOwner())
+	{
+		Character->OnCharacterMovementUpdated.AddDynamic(
+			this, 
+			&UCharacterMovementTrajectoryComponent::OnMovementUpdated);
+	}
+
 	Super::InitializeComponent();
 }
 
@@ -192,6 +291,11 @@ void UCharacterMovementTrajectoryComponent::UninitializeComponent()
 	Super::UninitializeComponent();
 }
 
+void UCharacterMovementTrajectoryComponent::OnMovementUpdated(float DeltaSeconds, FVector OldLocation, FVector OldVelocity)
+{
+	TickTrajectory(DeltaSeconds);
+}
+
 FTrajectorySampleRange UCharacterMovementTrajectoryComponent::GetTrajectory() const
 {
 	return GetTrajectoryWithSettings(PredictionSettings, bPredictionIncludesHistory);
@@ -213,7 +317,7 @@ FTrajectorySampleRange UCharacterMovementTrajectoryComponent::GetTrajectoryWithS
 	{
 		// Step the prediction iteratively towards the specified domain horizon(s)
 		FTrajectorySampleRange Prediction(SampleRate);
-		PredictTrajectory(MovementComponent, SampleRate, MaxSamples, Settings, PresentTrajectory, Prediction);
+		PredictTrajectory(MovementComponent, SampleRate, MaxSamples, Settings, PresentTrajectorySampleLS, Prediction);
 
 		// Combine past, present, and future into a uniformly sampled complete trajectory
 		return CombineHistoryPresentPrediction(bIncludeHistory, Prediction);
diff --git a/Engine/Plugins/Experimental/Animation/MotionTrajectory/Source/MotionTrajectory/Private/MotionTrajectoryLibrary.cpp b/Engine/Plugins/Experimental/Animation/MotionTrajectory/Source/MotionTrajectory/Private/MotionTrajectoryLibrary.cpp
--- a/Engine/Plugins/Experimental/Animation/MotionTrajectory/Source/MotionTrajectory/Private/MotionTrajectoryLibrary.cpp
+++ b/Engine/Plugins/Experimental/Animation/MotionTrajectory/Source/MotionTrajectory/Private/MotionTrajectoryLibrary.cpp
@@ -1,34 +1,50 @@
 // Copyright Epic Games, Inc. All Rights Reserved.
 
 #include "MotionTrajectoryLibrary.h"
+#include "MotionTrajectory.h"
 #include "GameFramework/Actor.h"
+#include "Components/SceneComponent.h"
 #include "Algo/Find.h"
 #include "GameFramework/Actor.h"
 
-static void FlattenTrajectoryPosition(FTrajectorySample& Sample, float& AccumulatedDeltaDistance, bool PreserveSpeed)
+static void FlattenTrajectoryPosition(
+	FTrajectorySample& Sample, 
+	const FTrajectorySample& PrevSample, 
+	const FTrajectorySample& FlattenedPrevSample, 
+	bool PreserveSpeed)
 {
-	if (!Sample.Position.IsZero())
+	if (!Sample.Transform.GetTranslation().IsZero())
 	{
+		const FVector Translation = Sample.Transform.GetLocation() - PrevSample.Transform.GetLocation();
+		const FVector FlattenedTranslation = FVector(Translation.X, Translation.Y, 0.0f);
+
 		if (PreserveSpeed)
 		{
+			const float TargetDistance = 
+				FMath::Abs(Sample.AccumulatedDistance - PrevSample.AccumulatedDistance);
+			const FVector FlattenedTranslationDir = Translation.GetSafeNormal2D();
+
 			// Take the full displacement, effectively meaning that the Z axis never existed
-			Sample.Position = Sample.Position.Size() * Sample.Position.GetSafeNormal2D();
+			Sample.Transform.SetLocation(
+				FlattenedPrevSample.Transform.GetLocation() + (FlattenedTranslationDir * TargetDistance));
 		}
 		else
 		{
 			// Accumulate the delta displacement difference as a result of Z axis being removed
-			const float DeltaDistance = Sample.Position.Size() - Sample.Position.Size2D();
-			AccumulatedDeltaDistance += DeltaDistance;
+			const float DeltaSeconds = Sample.AccumulatedSeconds - PrevSample.AccumulatedSeconds;
+			const float DeltaDistance = DeltaSeconds >= 0.0f ?
+				FlattenedTranslation.Size() :
+				-FlattenedTranslation.Size();
 
-			// Remove the accumulated delta displacement from each sample on the timeline, and recompute the new projected position
-			Sample.AccumulatedDistance -= AccumulatedDeltaDistance;
-			Sample.Position = FMath::Abs(Sample.AccumulatedDistance) * Sample.Position.GetSafeNormal2D();
+			Sample.AccumulatedDistance = FlattenedPrevSample.AccumulatedDistance + DeltaDistance;
+			Sample.Transform.SetTranslation(FlattenedPrevSample.Transform.GetLocation() + FlattenedTranslation);
 		}
 	}
 }
 
-FTrajectorySampleRange UMotionTrajectoryBlueprintLibrary::FlattenTrajectory2D(FTrajectorySampleRange Trajectory, bool PreserveSpeed)
+FTrajectorySampleRange UMotionTrajectoryBlueprintLibrary::FlattenTrajectory2D(FTrajectorySampleRange Trajectory, bool bPreserveSpeed)
 {
+
 	if (!Trajectory.HasSamples())
 	{
 		return Trajectory;
@@ -45,17 +61,17 @@ FTrajectorySampleRange UMotionTrajectoryBlueprintLibrary::FlattenTrajectory2D(FT
 		// Note: As as a consequence of magnitude preservation, AccumulatedDistance alongside AccumulatedTime should not need modification
 
 		// Linear velocity Z-axis component removal
-		if (!Sample.LocalLinearVelocity.IsZero())
+		if (!Sample.LinearVelocity.IsZero())
 		{
-			const float VelMagnitude = PreserveSpeed ? Sample.LocalLinearVelocity.Size() : Sample.LocalLinearVelocity.Size2D();
-			Sample.LocalLinearVelocity = VelMagnitude * Sample.LocalLinearVelocity.GetSafeNormal2D();
+			const float VelMagnitude = bPreserveSpeed ? Sample.LinearVelocity.Size() : Sample.LinearVelocity.Size2D();
+			Sample.LinearVelocity = VelMagnitude * Sample.LinearVelocity.GetSafeNormal2D();
 		}
 
 		// Align linear acceleration via projection onto velocity
-		if (!Sample.LocalLinearVelocity.IsZero() && !Sample.LocalLinearAcceleration.IsZero())
+		if (!Sample.LinearVelocity.IsZero() && !Sample.LinearAcceleration.IsZero())
 		{
-			const float AccelMagnitude = PreserveSpeed ? Sample.LocalLinearAcceleration.Size() : Sample.LocalLinearAcceleration.Size2D();
-			Sample.LocalLinearAcceleration = AccelMagnitude * Sample.LocalLinearAcceleration.ProjectOnTo(Sample.LocalLinearVelocity).GetSafeNormal();
+			const float AccelMagnitude = bPreserveSpeed ? Sample.LinearAcceleration.Size() : Sample.LinearAcceleration.Size2D();
+			Sample.LinearAcceleration = AccelMagnitude * Sample.LinearAcceleration.ProjectOnTo(Sample.LinearVelocity).GetSafeNormal();
 		}
 	}
 
@@ -66,11 +82,18 @@ FTrajectorySampleRange UMotionTrajectoryBlueprintLibrary::FlattenTrajectory2D(FT
 
 	check(PresentSampleIdx != INDEX_NONE);
 
+	// the present location should be zero but let's, for sanity, assume it might not be
+	FTrajectorySample& PresentSample = Trajectory.Samples[PresentSampleIdx];
+	const FVector PresentSampleLocation = PresentSample.Transform.GetLocation();
+	PresentSample.Transform.SetLocation(FVector(PresentSampleLocation.X, PresentSampleLocation.Y, 0.0f));
+
 	// Walk all samples into the future, conditionally removing contribution of Z axis motion
-	float AccumulatedDeltaDistance = 0.f;
-	for (int32 Idx = PresentSampleIdx, Num = Trajectory.Samples.Num(); Idx < Num; ++Idx)
+	FTrajectorySample PrevSample = PresentSample;
+	for (int32 Idx = PresentSampleIdx + 1, Num = Trajectory.Samples.Num(); Idx < Num; ++Idx)
 	{
-		FlattenTrajectoryPosition(Trajectory.Samples[Idx], AccumulatedDeltaDistance, PreserveSpeed);
+		FTrajectorySample CurrentSample = Trajectory.Samples[Idx];
+		FlattenTrajectoryPosition(Trajectory.Samples[Idx], PrevSample, Trajectory.Samples[Idx - 1], bPreserveSpeed);
+		PrevSample = CurrentSample;
 	}
 
 	// There is a possibility history has not been computed yet
@@ -80,10 +103,12 @@ FTrajectorySampleRange UMotionTrajectoryBlueprintLibrary::FlattenTrajectory2D(FT
 	}
 
 	// Walk all samples in the past, conditionally removing the contribution of Z axis motion
-	AccumulatedDeltaDistance = 0.f;
+	PrevSample = PresentSample;
 	for (int32 Idx = PresentSampleIdx - 1, Begin = 0; Idx >= Begin; --Idx)
 	{
-		FlattenTrajectoryPosition(Trajectory.Samples[Idx], AccumulatedDeltaDistance, PreserveSpeed);
+		FTrajectorySample CurrentSample = Trajectory.Samples[Idx];
+		FlattenTrajectoryPosition(Trajectory.Samples[Idx], PrevSample, Trajectory.Samples[Idx + 1], bPreserveSpeed);
+		PrevSample = CurrentSample;
 	}
 
 	return Trajectory;
@@ -120,7 +145,7 @@ static FVector ClampDirection(const FVector InputVector, const TArray<FTrajector
 	return Output;
 }
 
-FTrajectorySampleRange UMotionTrajectoryBlueprintLibrary::ClampTrajectoryDirection(FTrajectorySampleRange Trajectory, const TArray<FTrajectoryDirectionClamp>& Directions)
+FTrajectorySampleRange UMotionTrajectoryBlueprintLibrary::ClampTrajectoryDirection(FTrajectorySampleRange Trajectory, const TArray<FTrajectoryDirectionClamp>& Directions, bool bPreserveRotation)
 {
 	if (Directions.IsEmpty())
 	{
@@ -145,28 +170,35 @@ FTrajectorySampleRange UMotionTrajectoryBlueprintLibrary::ClampTrajectoryDirecti
 
 	check(PresentSample)
 
-	if (!PresentSample->LocalLinearVelocity.IsZero())
+	if (!PresentSample->LinearVelocity.IsZero())
 	{
-		const FVector VelocityBasis = ClampDirection(PresentSample->LocalLinearVelocity, Directions).GetSafeNormal();
+		const FVector VelocityBasis = ClampDirection(PresentSample->LinearVelocity, Directions).GetSafeNormal();
 
 		for (auto& Sample : Trajectory.Samples)
 		{
 			// Align linear velocity onto the velocity basis to maintain the present intended direction, while retaining per-sample magnitude
-			if (!Sample.LocalLinearVelocity.IsZero())
+			if (!Sample.LinearVelocity.IsZero())
 			{
-				Sample.LocalLinearVelocity = Sample.LocalLinearVelocity.Size() * Sample.LocalLinearVelocity.ProjectOnTo(VelocityBasis).GetSafeNormal();
+				Sample.LinearVelocity = Sample.LinearVelocity.Size() * Sample.LinearVelocity.ProjectOnTo(VelocityBasis).GetSafeNormal();
 			}
 
 			// Align linear acceleration through projection onto the modified velocity
-			if (!Sample.LocalLinearVelocity.IsZero() && !Sample.LocalLinearAcceleration.IsZero())
+			if (!Sample.LinearVelocity.IsZero() && !Sample.LinearAcceleration.IsZero())
 			{
-				Sample.LocalLinearAcceleration = Sample.LocalLinearAcceleration.Size() * Sample.LocalLinearAcceleration.ProjectOnTo(Sample.LocalLinearVelocity).GetSafeNormal();
+				Sample.LinearAcceleration = Sample.LinearAcceleration.Size() * Sample.LinearAcceleration.ProjectOnTo(Sample.LinearVelocity).GetSafeNormal();
 			}
 
 			// Align the position path through projection onto the modified velocity
-			if (!Sample.LocalLinearVelocity.IsZero() && !Sample.Position.IsZero())
+			if (!Sample.LinearVelocity.IsZero() && !Sample.Transform.GetLocation().IsZero())
 			{
-				Sample.Position = FMath::Abs(Sample.AccumulatedDistance) * Sample.Position.ProjectOnTo(Sample.LocalLinearVelocity).GetSafeNormal();
+				Sample.Transform.SetLocation(
+					FMath::Abs(Sample.AccumulatedDistance) * 
+					Sample.Transform.GetLocation().ProjectOnTo(Sample.LinearVelocity).GetSafeNormal());
+			}
+
+			if (bPreserveRotation)
+			{
+				Sample.Transform.SetRotation(PresentSample->Transform.GetRotation());
 			}
 		}
 	}
@@ -174,6 +206,33 @@ FTrajectorySampleRange UMotionTrajectoryBlueprintLibrary::ClampTrajectoryDirecti
 	return Trajectory;
 }
 
+FTrajectorySampleRange UMotionTrajectoryBlueprintLibrary::RotateTrajectory(
+	FTrajectorySampleRange Trajectory, 
+	const FQuat& Rotation)
+{
+	Trajectory.Rotate(Rotation);
+	return Trajectory;
+}
+
+FTrajectorySampleRange UMotionTrajectoryBlueprintLibrary::MakeTrajectoryRelativeToComponent(
+	FTrajectorySampleRange ActorTrajectory, 
+	const USceneComponent* Component)
+{
+	if (!IsValid(Component))
+	{
+		UE_LOG(LogMotionTrajectory, Error, TEXT("Invalid component!"));
+		return ActorTrajectory;
+	}
+
+	const AActor* Owner = Component->GetOwner();
+	const FTransform OwnerTransformWS = Owner->GetActorTransform();
+	const FTransform ComponentTransformWS = Component->GetComponentTransform();
+	const FQuat Rotation = ComponentTransformWS.GetRotation().Inverse() * OwnerTransformWS.GetRotation();
+
+	ActorTrajectory.Rotate(Rotation);
+	return ActorTrajectory;
+}
+
 void UMotionTrajectoryBlueprintLibrary::DebugDrawTrajectory(const AActor* Actor
 	, const FTransform& WorldTransform
 	, const FTrajectorySampleRange& Trajectory
diff --git a/Engine/Plugins/Experimental/Animation/MotionTrajectory/Source/MotionTrajectory/Public/MotionTrajectory.h b/Engine/Plugins/Experimental/Animation/MotionTrajectory/Source/MotionTrajectory/Public/MotionTrajectory.h
--- a/Engine/Plugins/Experimental/Animation/MotionTrajectory/Source/MotionTrajectory/Public/MotionTrajectory.h
+++ b/Engine/Plugins/Experimental/Animation/MotionTrajectory/Source/MotionTrajectory/Public/MotionTrajectory.h
@@ -10,6 +10,8 @@
 
 #include "MotionTrajectory.generated.h"
 
+DECLARE_LOG_CATEGORY_EXTERN(LogMotionTrajectory, Log, All);
+
 class APawn;
 
 class FMotionTrajectoryModule : public IModuleInterface
@@ -66,13 +68,18 @@ protected:
 	TRingBuffer<FTrajectorySample> SampleHistory = {};
 
 	// Retained instantaneous/present trajectory
-	FTrajectorySample PresentTrajectory = {};
+	FTrajectorySample PresentTrajectorySampleWS = {};
+
+	FTrajectorySample PresentTrajectorySampleLS = {};
+
+	// if anything makes the trajectory tick before ::TickComponent, skip it.
+	uint32 LastTrajectoryTickFrame = 0u;
 
 	// Gets the instantaneous/present trajectory sample of the current frame
-	virtual FTrajectorySample GetPresentTrajectory() const;
-	
-	// Gets the instantaneous/present world-space transform of the current frame
-	virtual FTransform GetPresentWorldTransform() const;
+	virtual FTrajectorySample CalcWorldSpacePresentTrajectorySample(float DeltaTime) const;
+
+	// Ticks the trajectory. Call this to tick before ::TickComponent when necessary
+	virtual void TickTrajectory(float DeltaTime);
 
 	// Gets the Pawn from the owning Actor of this component
 	const APawn* TryGetOwnerPawn() const;
@@ -80,6 +87,9 @@ protected:
 	// Combines all trajectory samples in the past, present, and future into a unifed trajectory range
 	FTrajectorySampleRange CombineHistoryPresentPrediction(bool bIncludeHistory, const FTrajectorySampleRange& Prediction) const;
 
+	// Checks the component internal frame counter to say if the component has already ticked this frame or not
+	bool HasTrajectoryTickedThisFrame() const;
+
 	// Forcefully evicts all trajectory history and resets internal history tracking state
 	void FlushHistory();
 
@@ -89,6 +99,7 @@ public:
 
 	// Begin UActorComponent Interface
 	virtual void OnComponentCreated() override;
+	virtual void BeginPlay() override;
 	virtual void TickComponent(float DeltaTime
 		, enum ELevelTick TickType
 		, FActorComponentTickFunction* ThisTickFunction) override;
diff --git a/Engine/Plugins/Experimental/Animation/MotionTrajectory/Source/MotionTrajectory/Public/MotionTrajectoryCharacterMovement.h b/Engine/Plugins/Experimental/Animation/MotionTrajectory/Source/MotionTrajectory/Public/MotionTrajectoryCharacterMovement.h
--- a/Engine/Plugins/Experimental/Animation/MotionTrajectory/Source/MotionTrajectory/Public/MotionTrajectoryCharacterMovement.h
+++ b/Engine/Plugins/Experimental/Animation/MotionTrajectory/Source/MotionTrajectory/Public/MotionTrajectoryCharacterMovement.h
@@ -15,10 +15,12 @@ class MOTIONTRAJECTORY_API UCharacterMovementTrajectoryComponent : public UMotio
 protected:
 
 	// Begin UMotionTrajectoryComponent Interface
-	virtual FTrajectorySample GetPresentTrajectory() const override;
-	virtual FTransform GetPresentWorldTransform() const override;
+	virtual FTrajectorySample CalcWorldSpacePresentTrajectorySample(float DeltaTime) const override;
 	// End UMotionTrajectoryComponent Interface
 
+	UFUNCTION()
+	void OnMovementUpdated(float DeltaSeconds, FVector OldLocation, FVector OldVelocity);
+
 public:
 
 	UCharacterMovementTrajectoryComponent(const FObjectInitializer& ObjectInitializer);
diff --git a/Engine/Plugins/Experimental/Animation/MotionTrajectory/Source/MotionTrajectory/Public/MotionTrajectoryLibrary.h b/Engine/Plugins/Experimental/Animation/MotionTrajectory/Source/MotionTrajectory/Public/MotionTrajectoryLibrary.h
--- a/Engine/Plugins/Experimental/Animation/MotionTrajectory/Source/MotionTrajectory/Public/MotionTrajectoryLibrary.h
+++ b/Engine/Plugins/Experimental/Animation/MotionTrajectory/Source/MotionTrajectory/Public/MotionTrajectoryLibrary.h
@@ -45,11 +45,39 @@ public:
 	*
 	* @param Trajectory			Input trajectory range
 	* @param Directions			Input direction clamping, containing angle thresholds for determining source to target direction
+	* @param bPreserveRotation	If true, sample rotations will be replaced with the present sample rotation
 	*
 	* @return					Direction clamped, modified trajectory range
 	*/
 	UFUNCTION(BlueprintPure, Category="Motion Trajectory", meta=(BlueprintThreadSafe, AutoCreateRefTerm="Directions"))
-	static FTrajectorySampleRange ClampTrajectoryDirection(FTrajectorySampleRange Trajectory, const TArray<FTrajectoryDirectionClamp>& Directions);
+	static FTrajectorySampleRange ClampTrajectoryDirection(
+		FTrajectorySampleRange Trajectory, 
+		const TArray<FTrajectoryDirectionClamp>& Directions, 
+		bool bPreserveRotation = true);
+
+	/**
+	* Rotates the trajectory
+	*
+	* @param Trajectory			Input trajectory range
+	* @param Rotation			Rotation to be applied to all trajectory samples
+	*
+	* @return					Rotated trajectory
+	*/
+	UFUNCTION(BlueprintPure, Category="Motion Trajectory", meta=(BlueprintThreadSafe))
+	static FTrajectorySampleRange RotateTrajectory(FTrajectorySampleRange Trajectory, const FQuat& Rotation);
+
+	/**
+	* Rotates the trajectory to make it relative to the component specified as a parameter
+	*
+	* @param Trajectory			Input trajectory range
+	* @param Component			Component whose transform will be the new trajectory reference frame
+	*
+	* @return					Trajectory relative to Component
+	*/
+	UFUNCTION(BlueprintPure, Category="Motion Trajectory", meta=(BlueprintThreadSafe))
+	static FTrajectorySampleRange MakeTrajectoryRelativeToComponent(
+		FTrajectorySampleRange ActorTrajectory, 
+		const class USceneComponent* Component);
 
 	/**
 	* Projects trajectory samples onto a defined set of allowed directions
@@ -63,13 +91,13 @@ public:
 	* @param ArrowThickness		Input sample arrow draw thickness
 	* @param bDrawText			Input include drawing of per-sample trajectory information
 	*/
-	UFUNCTION(BlueprintPure, Category="Motion Trajectory", meta=(BlueprintThreadSafe, AutoCreateRefTerm="WorldTransform"))
+	UFUNCTION(BlueprintCallable, Category="Motion Trajectory", meta=(BlueprintThreadSafe, AutoCreateRefTerm="WorldTransform"))
 	static void DebugDrawTrajectory(const AActor* Actor
 		, const FTransform& WorldTransform
 		, const FTrajectorySampleRange& Trajectory
 		, const FLinearColor PredictionColor = FLinearColor(0.f, 1.f, 0.f)
 		, const FLinearColor HistoryColor = FLinearColor(0.f, 0.f, 1.f)
-		, float ArrowScale = 0.025f
+		, float ArrowScale = 10.f
 		, float ArrowSize = 40.f
 		, float ArrowThickness = 2.f
 	);
diff --git a/Engine/Plugins/Experimental/Animation/PoseSearch/Source/Runtime/Private/PoseSearch.cpp b/Engine/Plugins/Experimental/Animation/PoseSearch/Source/Runtime/Private/PoseSearch.cpp
--- a/Engine/Plugins/Experimental/Animation/PoseSearch/Source/Runtime/Private/PoseSearch.cpp
+++ b/Engine/Plugins/Experimental/Animation/PoseSearch/Source/Runtime/Private/PoseSearch.cpp
@@ -1275,10 +1275,10 @@ void FPoseSearchFeatureVectorBuilder::BuildFromTrajectoryTimeBased(const FTrajec
 		Feature.SubsampleIdx = Idx;
 
 		Feature.Type = EPoseSearchFeatureType::LinearVelocity;
-		SetVector(Feature, Sample.LocalLinearVelocity);
+		SetVector(Feature, Sample.LinearVelocity);
 
 		Feature.Type = EPoseSearchFeatureType::Position;
-		SetVector(Feature, Sample.Position);
+		SetVector(Feature, Sample.Transform.GetLocation());
 	}
 }
 
@@ -1298,10 +1298,10 @@ void FPoseSearchFeatureVectorBuilder::BuildFromTrajectoryDistanceBased(const FTr
 		Feature.SubsampleIdx = Idx;
 
 		Feature.Type = EPoseSearchFeatureType::LinearVelocity;
-		SetVector(Feature, Sample.LocalLinearVelocity);
+		SetVector(Feature, Sample.LinearVelocity);
 
 		Feature.Type = EPoseSearchFeatureType::Position;
-		SetVector(Feature, Sample.Position);
+		SetVector(Feature, Sample.Transform.GetLocation());
 	} 
 }
 
diff --git a/Engine/Source/Runtime/Engine/Private/Animation/MotionTrajectoryTypes.cpp b/Engine/Source/Runtime/Engine/Private/Animation/MotionTrajectoryTypes.cpp
--- a/Engine/Source/Runtime/Engine/Private/Animation/MotionTrajectoryTypes.cpp
+++ b/Engine/Source/Runtime/Engine/Private/Animation/MotionTrajectoryTypes.cpp
@@ -53,10 +53,12 @@ namespace
 bool FTrajectorySample::IsZeroSample() const
 {
 	// AccumulatedTime is specifically omitted here to allow for the zero sample semantic across an entire trajectory range
-	return LocalLinearVelocity.IsNearlyZero()
-		&& LocalLinearAcceleration.IsNearlyZero()
-		&& Position.IsNearlyZero()
-		&& FMath::IsNearlyZero(AccumulatedDistance);
+	return LinearVelocity.IsNearlyZero()
+		&& LinearAcceleration.IsNearlyZero()
+		&& Transform.GetTranslation().IsNearlyZero()
+		&& FMath::IsNearlyZero(AccumulatedDistance)
+		&& FMath::IsNearlyZero(AngularSpeed)
+		&& Transform.GetRotation().IsIdentity();
 }
 
 FTrajectorySample FTrajectorySample::Lerp(const FTrajectorySample& Sample, float Alpha) const
@@ -64,13 +66,31 @@ FTrajectorySample FTrajectorySample::Lerp(const FTrajectorySample& Sample, float
 	FTrajectorySample Interp;
 	Interp.AccumulatedSeconds = FMath::Lerp(AccumulatedSeconds, Sample.AccumulatedSeconds, Alpha);
 	Interp.AccumulatedDistance = FMath::Lerp(AccumulatedDistance, Sample.AccumulatedDistance, Alpha);
-	Interp.Position = FMath::Lerp(Position, Sample.Position, Alpha);
-	Interp.LocalLinearVelocity = FMath::Lerp(LocalLinearVelocity, Sample.LocalLinearVelocity, Alpha);
-	Interp.LocalLinearAcceleration = FMath::Lerp(LocalLinearAcceleration, Sample.LocalLinearAcceleration, Alpha);
+	Interp.LinearVelocity = FMath::Lerp(LinearVelocity, Sample.LinearVelocity, Alpha);
+	Interp.LinearAcceleration = FMath::Lerp(LinearAcceleration, Sample.LinearAcceleration, Alpha);
+
+	Interp.Transform.Blend(Transform, Sample.Transform, Alpha);
+	
+	// TODO: this is very simple, more like Lerp than Slerp, can we do better?
+	const FVector AngularVelocity = AngularVelocityAxis * AngularSpeed;
+	const FVector SampleAngularVelocity = Sample.AngularVelocityAxis * Sample.AngularSpeed;
+	const FVector LerpedAngularVelocity = FMath::Lerp(AngularVelocity, SampleAngularVelocity, Alpha);
+	const float LerpedAngularSpeed = LerpedAngularVelocity.Size();
+	if (LerpedAngularSpeed > SMALL_NUMBER)
+	{
+		Interp.AngularVelocityAxis = LerpedAngularVelocity / LerpedAngularSpeed;
+		Interp.AngularSpeed = LerpedAngularSpeed;
+	}
+	else
+	{
+		Interp.AngularVelocityAxis = FVector::ZeroVector;
+		Interp.AngularSpeed = 0.0f;
+	}
+
 	return Interp;
 }
 
-FTrajectorySample FTrajectorySample::CubicCRSplineInterp(const FTrajectorySample& PrevSample
+FTrajectorySample FTrajectorySample::SmoothInterp(const FTrajectorySample& PrevSample
 	, const FTrajectorySample& Sample
 	, const FTrajectorySample& NextSample
 	, float Alpha) const
@@ -78,12 +98,75 @@ FTrajectorySample FTrajectorySample::CubicCRSplineInterp(const FTrajectorySample
 	FTrajectorySample Interp;
 	Interp.AccumulatedDistance = CubicCRSplineInterpSafe(PrevSample.AccumulatedDistance, AccumulatedDistance, Sample.AccumulatedDistance, NextSample.AccumulatedDistance, Alpha);
 	Interp.AccumulatedSeconds = CubicCRSplineInterpSafe(PrevSample.AccumulatedSeconds, AccumulatedSeconds, Sample.AccumulatedSeconds, NextSample.AccumulatedSeconds, Alpha);
-	Interp.Position = CubicCRSplineInterpSafe(PrevSample.Position, Position, Sample.Position, NextSample.Position, Alpha);
-	Interp.LocalLinearVelocity = CubicCRSplineInterpSafe(PrevSample.LocalLinearVelocity, LocalLinearVelocity, Sample.LocalLinearVelocity, NextSample.LocalLinearVelocity, Alpha);
-	Interp.LocalLinearAcceleration = CubicCRSplineInterpSafe(PrevSample.LocalLinearAcceleration, LocalLinearAcceleration, Sample.LocalLinearAcceleration, NextSample.LocalLinearAcceleration, Alpha);
+	Interp.LinearVelocity = CubicCRSplineInterpSafe(PrevSample.LinearVelocity, LinearVelocity, Sample.LinearVelocity, NextSample.LinearVelocity, Alpha);
+	Interp.LinearAcceleration = CubicCRSplineInterpSafe(PrevSample.LinearAcceleration, LinearAcceleration, Sample.LinearAcceleration, NextSample.LinearAcceleration, Alpha);
+
+	Interp.Transform.SetLocation(CubicCRSplineInterpSafe(
+		PrevSample.Transform.GetLocation(),
+		Transform.GetLocation(),
+		Sample.Transform.GetLocation(),
+		NextSample.Transform.GetLocation(),
+		Alpha));
+	FQuat Q0 = PrevSample.Transform.GetRotation().W >= 0.0f ? 
+		PrevSample.Transform.GetRotation() : -PrevSample.Transform.GetRotation();
+	FQuat Q1 = Transform.GetRotation().W >= 0.0f ? 
+		Transform.GetRotation() : -Transform.GetRotation();
+	FQuat Q2 = Sample.Transform.GetRotation().W >= 0.0f ? 
+		Sample.Transform.GetRotation() : -Sample.Transform.GetRotation();
+	FQuat Q3 = NextSample.Transform.GetRotation().W >= 0.0f ? 
+		NextSample.Transform.GetRotation() : -NextSample.Transform.GetRotation();
+
+	FQuat T0, T1;
+	FQuat::CalcTangents(Q0, Q1, Q2, 0.0f, T0);
+	FQuat::CalcTangents(Q1, Q2, Q3, 0.0f, T1);
+
+	Interp.Transform.SetRotation(FQuat::Squad(Q1, T0, Q2, T1, Alpha));
+
+	const FVector V0 = PrevSample.AngularVelocityAxis * PrevSample.AngularSpeed;
+	const FVector V1 = AngularVelocityAxis * AngularSpeed;
+	const FVector V2 = Sample.AngularVelocityAxis * Sample.AngularSpeed;
+	const FVector V3 = NextSample.AngularVelocityAxis * NextSample.AngularSpeed;
+
+	const FVector LerpedAngularVelocity = CubicCRSplineInterpSafe(V0, V1, V2, V3, Alpha);
+	const float LerpedAngularSpeed = LerpedAngularVelocity.Size();
+	if (LerpedAngularSpeed > SMALL_NUMBER)
+	{
+		Interp.AngularVelocityAxis = LerpedAngularVelocity / LerpedAngularSpeed;
+		Interp.AngularSpeed = LerpedAngularSpeed;
+	}
+	else
+	{
+		Interp.AngularVelocityAxis = FVector::ZeroVector;
+		Interp.AngularSpeed = 0.0f;
+	}
+
 	return Interp;
 }
 
+void FTrajectorySample::PrependOffset(const FTransform DeltaTransform, float DeltaSeconds)
+{
+	AccumulatedSeconds += DeltaSeconds;
+
+	if (FMath::IsNearlyZero(AccumulatedSeconds))
+	{
+		AccumulatedDistance = 0.0f;
+	}
+	else
+	{
+		const float DistanceOffset = DeltaSeconds >= 0.0f ?
+			DeltaTransform.GetTranslation().Size() :
+			-DeltaTransform.GetTranslation().Size();
+
+		AccumulatedDistance += DistanceOffset;
+	}
+
+	Transform *= DeltaTransform;
+
+	LinearVelocity = DeltaTransform.TransformVectorNoScale(LinearVelocity);
+	LinearAcceleration = DeltaTransform.TransformVectorNoScale(LinearAcceleration);
+	AngularVelocityAxis = DeltaTransform.TransformVectorNoScale(AngularVelocityAxis);
+}
+
 bool FTrajectorySampleRange::HasSamples() const
 {
 	return !Samples.IsEmpty();
@@ -105,6 +188,14 @@ void FTrajectorySampleRange::RemoveHistory()
 		});
 }
 
+void FTrajectorySampleRange::Rotate(const FQuat& Rotation)
+{
+	for (auto& Sample : Samples)
+	{
+		Sample.PrependOffset(FTransform(Rotation), 0.0f);
+	}
+}
+
 void FTrajectorySampleRange::DebugDrawTrajectory(bool bEnable
 	, const UWorld* World
 	, const FTransform& WorldTransform
@@ -128,14 +219,17 @@ void FTrajectorySampleRange::DebugDrawTrajectory(bool bEnable
 #endif
 			for (int32 Idx = 0, Num = Samples.Num(); Idx < Num; Idx++)
 			{
-				const FVector WorldPosition = WorldTransform.TransformPosition(Samples[Idx].Position);
-				const FVector WorldVelocity = (WorldTransform.TransformVector(Samples[Idx].LocalLinearVelocity) * ArrowScale) + WorldPosition;
+				const FVector WorldPosition = WorldTransform.TransformPosition(Samples[Idx].Transform.GetTranslation());
+				const FVector WorldForward = 
+					(WorldTransform.TransformVectorNoScale(Samples[Idx].Transform.GetRotation().GetAxisX()) *
+					 ArrowScale) 
+					+ WorldPosition;
 
 				// Interpolate the history and prediction color over the entire trajectory range
 				const float ColorLerp = static_cast<float>(Idx) / static_cast<float>(Num);
 				const FLinearColor Color = FLinearColor::LerpUsingHSV(PredictionColor, HistoryColor, ColorLerp);
 
-				DrawDebugDirectionalArrow(World, WorldPosition, WorldVelocity, ArrowSize, Color.ToFColor(true), false, 0.f, 0, ArrowThickness);
+				DrawDebugDirectionalArrow(World, WorldPosition, WorldForward, ArrowSize, Color.ToFColor(true), false, 0.f, 0, ArrowThickness);
 #if ENABLE_ANIM_DEBUG
 				FString DebugString;
 				FString DebugSampleString;
@@ -155,15 +249,15 @@ void FTrajectorySampleRange::DebugDrawTrajectory(bool bEnable
 					break;
 				case 4: // Sample Position
 					DebugString = "Sample Position:";
-					DebugSampleString = DebugSampleString.Format(TEXT("{0}"), { Samples[Idx].Position.ToCompactString() });
+					DebugSampleString = DebugSampleString.Format(TEXT("{0}"), { Samples[Idx].Transform.GetLocation().ToCompactString() });
 					break;
 				case 5: // Sample Velocity
 					DebugString = "Sample Velocity:";
-					DebugSampleString = DebugSampleString.Format(TEXT("{0}"), { Samples[Idx].LocalLinearVelocity.ToCompactString() });
+					DebugSampleString = DebugSampleString.Format(TEXT("{0}"), { Samples[Idx].LinearVelocity.ToCompactString() });
 					break;
 				case 6: // Sample Acceleration
 					DebugString = "Sample Acceleration:";
-					DebugSampleString = DebugSampleString.Format(TEXT("{0}"), { Samples[Idx].LocalLinearAcceleration.ToCompactString() });
+					DebugSampleString = DebugSampleString.Format(TEXT("{0}"), { Samples[Idx].LinearAcceleration.ToCompactString() });
 					break;
 				default:
 					break;
@@ -178,7 +272,7 @@ void FTrajectorySampleRange::DebugDrawTrajectory(bool bEnable
 						DrawDebugString(World, WorldTransform.GetLocation() + FVector(0.f, 0.f, 50.f), DebugString, nullptr, FColor::White, 0.f, false, 1.f);
 					}
 
-					DrawDebugString(World, WorldVelocity + FVector(0.f, 0.f, 10.f), DebugSampleString, nullptr, FColor::White, 0.f, false, 1.f);
+					DrawDebugString(World, WorldForward + FVector(0.f, 0.f, 10.f), DebugSampleString, nullptr, FColor::White, 0.f, false, 1.f);
 				}
 #endif
 			}
diff --git a/Engine/Source/Runtime/Engine/Public/Animation/MotionTrajectoryTypes.h b/Engine/Source/Runtime/Engine/Public/Animation/MotionTrajectoryTypes.h
--- a/Engine/Source/Runtime/Engine/Public/Animation/MotionTrajectoryTypes.h
+++ b/Engine/Source/Runtime/Engine/Public/Animation/MotionTrajectoryTypes.h
@@ -34,25 +34,37 @@ struct ENGINE_API FTrajectorySample
 
 	// Position relative to the sampled in-motion object
 	UPROPERTY(EditAnywhere, BlueprintReadWrite, Category="Motion Trajectory")
-	FVector Position = FVector::ZeroVector;
+	FTransform Transform = FTransform::Identity;
 
 	// Linear velocity relative to the sampled in-motion object
 	UPROPERTY(EditAnywhere, BlueprintReadWrite, Category="Motion Trajectory")
-	FVector LocalLinearVelocity = FVector::ZeroVector;
+	FVector LinearVelocity = FVector::ZeroVector;
+
+	// Axis around which the angular velocity rotates the in-motion object
+	UPROPERTY(EditAnywhere, BlueprintReadWrite, Category = "Motion Trajectory")
+	FVector AngularVelocityAxis = FVector::ZeroVector;
+
+	// Rotation speed around AngularVelocityAxis
+	UPROPERTY(EditAnywhere, BlueprintReadWrite, Category = "Motion Trajectory")
+	float AngularSpeed = 0.0f;
 
 	// Linear acceleration relative to the sampled in-motion object
 	UPROPERTY(EditAnywhere, BlueprintReadWrite, Category="Motion Trajectory")
-	FVector LocalLinearAcceleration = FVector::ZeroVector;
-	
+	FVector LinearAcceleration = FVector::ZeroVector;
+
 	// Linear interpolation of all parameters of two trajectory samples
 	FTrajectorySample Lerp(const FTrajectorySample& Sample, float Alpha) const;
 
 	// Centripetal Catmull–Rom spline interpolation of all parameters of two trajectory samples
-	FTrajectorySample CubicCRSplineInterp(const FTrajectorySample& PrevSample
+	FTrajectorySample SmoothInterp(const FTrajectorySample& PrevSample
 		, const FTrajectorySample& Sample
 		, const FTrajectorySample& NextSample
 		, float Alpha) const;
 
+	// Concatenates DeltaTransform before the current transform is applied and shifts the accumulated time by 
+	// DeltaSeconds
+	void PrependOffset(const FTransform DeltaTransform, float DeltaSeconds);
+
 	// Determines if all sample properties are zeroed
 	bool IsZeroSample() const;
 };
@@ -66,7 +78,7 @@ struct ENGINE_API FTrajectorySampleRange
 	// Debug rendering contants
 	static constexpr FLinearColor DebugDefaultPredictionColor{ 0.f, 1.f, 0.f };
 	static constexpr FLinearColor DebugDefaultHistoryColor{ 0.f, 0.f, 1.f };
-	static constexpr float DebugDefaultArrowScale = 0.025f;
+	static constexpr float DebugDefaultArrowScale = 10.f;
 	static constexpr float DebugDefaultArrowSize = 40.f;
 	static constexpr float DebugDefaultArrowThickness = 2.f;
 
@@ -86,6 +98,9 @@ struct ENGINE_API FTrajectorySampleRange
 	// Removes history samples from trajectory (retains present and future)
 	void RemoveHistory();
 
+	// Rotates all samples in the trajectory
+	void Rotate(const FQuat& Rotation);
+
 	// Determine if any trajectory samples are present
 	bool HasSamples() const;
 
@@ -146,7 +161,7 @@ struct ENGINE_API FTrajectorySampleRange
 
 					if (PrevIdx != InitialIdx && InitialIdx != Idx && Idx != NextIdx)
 					{
-						InterpSample = InitialSample.CubicCRSplineInterp(Samples[PrevIdx], NextSample, Samples[NextIdx], Alpha);
+						InterpSample = InitialSample.SmoothInterp(Samples[PrevIdx], NextSample, Samples[NextIdx], Alpha);
 						return InterpSample;
 					}
 				}
diff --git a/Engine/Source/Runtime/Experimental/Animation/DynamicPlayRate/Private/DynamicPlayRate/DynamicPlayRateLibrary.cpp b/Engine/Source/Runtime/Experimental/Animation/DynamicPlayRate/Private/DynamicPlayRate/DynamicPlayRateLibrary.cpp
--- a/Engine/Source/Runtime/Experimental/Animation/DynamicPlayRate/Private/DynamicPlayRate/DynamicPlayRateLibrary.cpp
+++ b/Engine/Source/Runtime/Experimental/Animation/DynamicPlayRate/Private/DynamicPlayRate/DynamicPlayRateLibrary.cpp
@@ -93,7 +93,7 @@ float DynamicPlayRateAdjustment(const FAnimationUpdateContext& Context
 	// Approximate zeroed values may indicate the synchronization point for a stop or pivot
 	const FTrajectorySample* MinimaSample = Algo::MinElementBy(Trajectory.Samples, [](const FTrajectorySample& Sample)
 		{
-			return Sample.LocalLinearVelocity.SizeSquared();
+			return Sample.LinearVelocity.SizeSquared();
 		});
 
 	// Given a high-resolution time step, walk the current animation sequence to find a corresponding minima root motion delta
@@ -170,7 +170,7 @@ float DynamicPlayRateAdjustment(const FAnimationUpdateContext& Context
 		// Minima driven play rate scaling synchronizes using the remaining displacement to near zero with: animation / locomotion
 		// Locomotion driven play rate scaling synchronizes using the per-frame instantaneous displacement with: locomotion / animation
 		const float SequenceDelta = bRootMotionDrivenPlayRate ? MinimaSampleTime - AccumulatedTime : DeltaTime;
-		const float TrajectoryDisplacement = bRootMotionDrivenPlayRate ? MinimaSample->AccumulatedDistance : Trajectory.Samples[0].LocalLinearVelocity.Size() * DeltaTime;
+		const float TrajectoryDisplacement = bRootMotionDrivenPlayRate ? MinimaSample->AccumulatedDistance : Trajectory.Samples[0].LinearVelocity.Size() * DeltaTime;
 
 		RootMotionDelta = InternalSequence->ExtractRootMotion(AccumulatedTime, SequenceDelta, bLooping).GetTranslation();
 		const float RootMotionDeltaDisplacement = RootMotionDelta.Size();
@@ -224,7 +224,9 @@ float DynamicPlayRateAdjustment(const FAnimationUpdateContext& Context
 
 		if (bRootMotionDrivenPlayRate)
 		{
-			Context.AnimInstanceProxy->AnimDrawDebugSphere(Context.AnimInstanceProxy->GetComponentTransform().TransformPosition(MinimaSample->Position), 8.f, 16, FColor::Yellow);
+			Context.AnimInstanceProxy->AnimDrawDebugSphere(
+				Context.AnimInstanceProxy->GetComponentTransform().TransformPosition(MinimaSample->Transform.GetLocation()),
+				8.f, 16, FColor::Yellow);
 		}
 	}
 #endif
