Trajectory prediction merges orient to body with controller rotation instead of picking one or the other

diff --git a/Engine/Plugins/Experimental/Animation/MotionTrajectory/Source/MotionTrajectory/Private/MotionTrajectory.cpp b/Engine/Plugins/Experimental/Animation/MotionTrajectory/Source/MotionTrajectory/Private/MotionTrajectory.cpp
--- a/Engine/Plugins/Experimental/Animation/MotionTrajectory/Source/MotionTrajectory/Private/MotionTrajectory.cpp
+++ b/Engine/Plugins/Experimental/Animation/MotionTrajectory/Source/MotionTrajectory/Private/MotionTrajectory.cpp
@@ -33,15 +33,8 @@ FTrajectorySample UMotionTrajectoryComponent::CalcWorldSpacePresentTrajectorySam
 	return FTrajectorySample();
 }
 
-bool UMotionTrajectoryComponent::HasTrajectoryTickedThisFrame() const
-{
-	return GFrameCounter == LastTrajectoryTickFrame;
-}
-
 void UMotionTrajectoryComponent::TickTrajectory(float DeltaTime)
 {
-	LastTrajectoryTickFrame = static_cast<uint32>(GFrameCounter);
-
 	// Compute the instantaneous/present trajectory sample and guarantee that is zeroed on all accumulated domains
 	PresentTrajectorySampleWS = CalcWorldSpacePresentTrajectorySample(DeltaTime);
 	PresentTrajectorySampleWS.AccumulatedDistance = 0.f;
@@ -163,12 +156,7 @@ void UMotionTrajectoryComponent::TickHistoryEvictionPolicy()
 			});
 	}
 
-	// Disallow trajectory samples beyond the maximum count
-	// An opportunity to insert more samples will arise as history decays
-	if (SampleHistory.Num() < MaxSamples)
-	{
-		SampleHistory.Emplace(PresentTrajectorySampleLS);
-	}
+	SampleHistory.Emplace(PresentTrajectorySampleLS);
 	
 	// Cache the present world space sample as the next frame's reference point for delta computation
 	PreviousWorldTransform = WorldTransform;
@@ -205,18 +193,13 @@ void UMotionTrajectoryComponent::BeginPlay()
 	PresentTrajectorySampleWS.Transform = PresentTransform;
 
 	PresentTrajectorySampleLS = FTrajectorySample();
-
-	LastTrajectoryTickFrame = 0u;
 }
 
 void UMotionTrajectoryComponent::TickComponent(float DeltaTime, enum ELevelTick TickType, FActorComponentTickFunction* ThisTickFunction)
 {
 	Super::TickComponent(DeltaTime, TickType, ThisTickFunction);
 
-	if (!HasTrajectoryTickedThisFrame())
-	{
-		TickTrajectory(DeltaTime);
-	}
+	TickTrajectory(DeltaTime);
 }
 
 FTrajectorySampleRange UMotionTrajectoryComponent::GetTrajectory() const
diff --git a/Engine/Plugins/Experimental/Animation/MotionTrajectory/Source/MotionTrajectory/Private/MotionTrajectoryCharacterMovement.cpp b/Engine/Plugins/Experimental/Animation/MotionTrajectory/Source/MotionTrajectory/Private/MotionTrajectoryCharacterMovement.cpp
--- a/Engine/Plugins/Experimental/Animation/MotionTrajectory/Source/MotionTrajectory/Private/MotionTrajectoryCharacterMovement.cpp
+++ b/Engine/Plugins/Experimental/Animation/MotionTrajectory/Source/MotionTrajectory/Private/MotionTrajectoryCharacterMovement.cpp
@@ -8,19 +8,81 @@
 #include "GameFramework/Character.h"
 #include "KismetAnimationLibrary.h"
 
-// ----------- BEGIN Derived from FCharacterMovementComponentAsyncInput::CalcVelocity() ----------- //
+// ----------- BEGIN Derived from FCharacterMovementComponentAsyncInput ----------- //
+// FCharacterMovementComponentAsyncInput::CalcVelocity() for linear motion
+// FCharacterMovementComponentAsyncInput::PhysicsRotation() for angular motion
+
+static FRotator AdjustRotationVerticality(
+	const FRotator& InRotation, 
+	const UCharacterMovementComponent* MovementComponent)
+{
+	FRotator OutRotation = InRotation;
+
+	if (MovementComponent->ShouldRemainVertical())
+	{
+		OutRotation.Pitch = 0.f;
+		OutRotation.Yaw = FRotator::NormalizeAxis(OutRotation.Yaw);
+		OutRotation.Roll = 0.f;
+	}
+	else
+	{
+		OutRotation.Normalize();
+	}
+
+	return OutRotation;
+}
+
+static FRotator LimitRotationRate(
+	const FRotator& InRotation,
+	const FRotator& CurrentRotation,
+	const FRotator& DeltaRot,
+	const UCharacterMovementComponent* MovementComponent)
+{
+	FRotator OutRotation = InRotation;
+
+	// Accumulate a desired new rotation.
+	const float AngleTolerance = 1e-3f;
+
+	if (!CurrentRotation.Equals(OutRotation, AngleTolerance))
+	{
+		// PITCH
+		if (!FMath::IsNearlyEqual(CurrentRotation.Pitch, OutRotation.Pitch, AngleTolerance))
+		{
+			OutRotation.Pitch = FMath::FixedTurn(CurrentRotation.Pitch, OutRotation.Pitch, DeltaRot.Pitch);
+		}
+
+		// YAW
+		if (!FMath::IsNearlyEqual(CurrentRotation.Yaw, OutRotation.Yaw, AngleTolerance))
+		{
+			OutRotation.Yaw = FMath::FixedTurn(CurrentRotation.Yaw, OutRotation.Yaw, DeltaRot.Yaw);
+		}
+
+		// ROLL
+		if (!FMath::IsNearlyEqual(CurrentRotation.Roll, OutRotation.Roll, AngleTolerance))
+		{
+			OutRotation.Roll = FMath::FixedTurn(CurrentRotation.Roll, OutRotation.Roll, DeltaRot.Roll);
+		}
+
+		// Set the new rotation.
+		OutRotation.DiagnosticCheckNaN(TEXT("LimitRotationRate: OutRotation"));
+	}
+
+	return OutRotation;
+}
 
 static void StepPrediction(float IntegrationDelta
 	, const UCharacterMovementComponent* MovementComponent
-	, FTrajectorySample& Sample)
+	, const FRotator& ControlRotationVelocity
+	, FRotator& InOutControlRotationTotalDelta
+	, FTrajectorySample& InOutSample)
 {
-	const FTransform InitialSampleTransform = Sample.Transform;
+	const FTransform InitialSampleTransform = InOutSample.Transform;
 
 	if (MovementComponent->GetCurrentAcceleration().IsZero())
 	{
 		float ActualBrakingFriction = (MovementComponent->bUseSeparateBrakingFriction ? MovementComponent->BrakingFriction : MovementComponent->GroundFriction);
 
-		if (!Sample.LinearVelocity.IsZero())
+		if (!InOutSample.LinearVelocity.IsZero())
 		{
 			const float FrictionFactor = FMath::Max(0.f, MovementComponent->BrakingFrictionFactor);
 			ActualBrakingFriction = FMath::Max(0.f, ActualBrakingFriction * FrictionFactor);
@@ -37,127 +99,80 @@ static void StepPrediction(float IntegrationDelta
 			float RemainingTime = IntegrationDelta;
 			const float MaxTimeStep = FMath::Clamp(MovementComponent->BrakingSubStepTime, 1.f / 75.f, 1.f / 20.f);
 
-			const FVector PrevLinearVelocity = Sample.LinearVelocity;
-			const FVector RevAccel = (bZeroBraking ? FVector::ZeroVector : (-BrakingDeceleration * Sample.LinearVelocity.GetSafeNormal()));
+			const FVector PrevLinearVelocity = InOutSample.LinearVelocity;
+			const FVector RevAccel = (bZeroBraking ? FVector::ZeroVector : (-BrakingDeceleration * InOutSample.LinearVelocity.GetSafeNormal()));
 
 			while (RemainingTime >= UCharacterMovementComponent::MIN_TICK_TIME)
 			{
 				const float dt = ((RemainingTime > MaxTimeStep && !bZeroFriction) ? FMath::Min(MaxTimeStep, RemainingTime * 0.5f) : RemainingTime);
 				RemainingTime -= dt;
 
-				Sample.LinearVelocity = Sample.LinearVelocity + ((-ActualBrakingFriction) * Sample.LinearVelocity + RevAccel) * dt;
+				InOutSample.LinearVelocity = InOutSample.LinearVelocity + ((-ActualBrakingFriction) * InOutSample.LinearVelocity + RevAccel) * dt;
 
-				if ((Sample.LinearVelocity | PrevLinearVelocity) <= 0.f)
+				if ((InOutSample.LinearVelocity | PrevLinearVelocity) <= 0.f)
 				{
-					Sample.LinearVelocity = FVector::ZeroVector;
+					InOutSample.LinearVelocity = FVector::ZeroVector;
 					return;
 				}
 			}
 
 			// Clamp to zero if nearly zero, or if below min threshold and braking
-			const float VSizeSq = Sample.LinearVelocity.SizeSquared();
+			const float VSizeSq = InOutSample.LinearVelocity.SizeSquared();
 			if (VSizeSq <= KINDA_SMALL_NUMBER || (!bZeroBraking && VSizeSq <= FMath::Square(UCharacterMovementComponent::BRAKE_TO_STOP_VELOCITY)))
 			{
-				Sample.LinearVelocity = FVector::ZeroVector;
+				InOutSample.LinearVelocity = FVector::ZeroVector;
 			}
 		}
 	}
 	else
 	{
-		const FVector AccelDir = Sample.LinearAcceleration.GetSafeNormal();
-		const float VelSize = Sample.LinearVelocity.Size();
-		Sample.LinearVelocity = Sample.LinearVelocity - (Sample.LinearVelocity - AccelDir * VelSize) * FMath::Min(IntegrationDelta * MovementComponent->GroundFriction, 1.f);
+		const FVector AccelDir = InOutSample.LinearAcceleration.GetSafeNormal();
+		const float VelSize = InOutSample.LinearVelocity.Size();
+		InOutSample.LinearVelocity = InOutSample.LinearVelocity - (InOutSample.LinearVelocity - AccelDir * VelSize) * FMath::Min(IntegrationDelta * MovementComponent->GroundFriction, 1.f);
 
 		const float MaxInputSpeed = FMath::Max(MovementComponent->GetMaxSpeed() * MovementComponent->GetAnalogInputModifier(), MovementComponent->GetMinAnalogSpeed());
-		Sample.LinearVelocity += Sample.LinearAcceleration * IntegrationDelta;
-		Sample.LinearVelocity = Sample.LinearVelocity.GetClampedToMaxSize(MaxInputSpeed);
+		InOutSample.LinearVelocity += InOutSample.LinearAcceleration * IntegrationDelta;
+		InOutSample.LinearVelocity = InOutSample.LinearVelocity.GetClampedToMaxSize(MaxInputSpeed);
 	}
 
-	const FTransform InitialTransformWS = InitialSampleTransform * MovementComponent->GetOwner()->GetActorTransform();
+	const FTransform PrevTransformWS = InitialSampleTransform * MovementComponent->GetOwner()->GetActorTransform();
 
-	if (MovementComponent->bOrientRotationToMovement)
 	{
-		FRotator CurrentRotation = InitialTransformWS.GetRotation().Rotator();
-		CurrentRotation.DiagnosticCheckNaN(TEXT("CharacterMovementComponent::PhysicsRotation(): CurrentRotation"));
-
 		FRotator DeltaRot = MovementComponent->GetDeltaRotation(IntegrationDelta);
-		DeltaRot.DiagnosticCheckNaN(TEXT("CharacterMovementComponent::PhysicsRotation(): GetDeltaRotation"));
-
-		FRotator DesiredRotation = 
-			MovementComponent->ComputeOrientToMovementRotation(CurrentRotation, IntegrationDelta, DeltaRot);
-
-		// TODO: leaving this code commented here as a reference for later combining both modes so we have one 
-		// component of rotation targeting a fixed direction (above) and the below being extrapolated
-		//FRotator DesiredRotation = CurrentRotation;
-		//if (MovementComponent->bOrientRotationToMovement)
-		//{
-		//	DesiredRotation = MovementComponent->ComputeOrientToMovementRotation(CurrentRotation, IntegrationDelta, DeltaRot);
-		//}
-		//else if (MovementComponent->GetCharacterOwner()->Controller && MovementComponent->bUseControllerDesiredRotation)
-		//{
-		//	DesiredRotation = MovementComponent->GetCharacterOwner()->Controller->GetDesiredRotation();
-		//}
-
-		if (MovementComponent->ShouldRemainVertical())
+		DeltaRot.DiagnosticCheckNaN(TEXT("StepPrediction: MovementComponent->GetDeltaRotation"));
+
+		FRotator BaseRotation = FRotator::ZeroRotator; 
+		if (MovementComponent->bOrientRotationToMovement)
 		{
-			DesiredRotation.Pitch = 0.f;
-			DesiredRotation.Yaw = FRotator::NormalizeAxis(DesiredRotation.Yaw);
-			DesiredRotation.Roll = 0.f;
+			const FRotator PrevRotator = PrevTransformWS.GetRotation().Rotator() - InOutControlRotationTotalDelta;
+			PrevRotator.DiagnosticCheckNaN(TEXT("StepPrediction: PrevRotator"));
+
+			BaseRotation =
+				MovementComponent->ComputeOrientToMovementRotation(PrevRotator, IntegrationDelta, DeltaRot);
+			BaseRotation = AdjustRotationVerticality(BaseRotation, MovementComponent);
+			BaseRotation = LimitRotationRate(BaseRotation, PrevRotator, DeltaRot, MovementComponent);
 		}
 		else
 		{
-			DesiredRotation.Normalize();
+			BaseRotation = MovementComponent->GetOwner()->GetActorTransform().Rotator();
 		}
 
-		// Accumulate a desired new rotation.
-		const float AngleTolerance = 1e-3f;
-
-		if (!CurrentRotation.Equals(DesiredRotation, AngleTolerance))
-		{
-			// PITCH
-			if (!FMath::IsNearlyEqual(CurrentRotation.Pitch, DesiredRotation.Pitch, AngleTolerance))
-			{
-				DesiredRotation.Pitch = FMath::FixedTurn(CurrentRotation.Pitch, DesiredRotation.Pitch, DeltaRot.Pitch);
-			}
-
-			// YAW
-			if (!FMath::IsNearlyEqual(CurrentRotation.Yaw, DesiredRotation.Yaw, AngleTolerance))
-			{
-				DesiredRotation.Yaw = FMath::FixedTurn(CurrentRotation.Yaw, DesiredRotation.Yaw, DeltaRot.Yaw);
-			}
-
-			// ROLL
-			if (!FMath::IsNearlyEqual(CurrentRotation.Roll, DesiredRotation.Roll, AngleTolerance))
-			{
-				DesiredRotation.Roll = FMath::FixedTurn(CurrentRotation.Roll, DesiredRotation.Roll, DeltaRot.Roll);
-			}
-
-			// Set the new rotation.
-			DesiredRotation.DiagnosticCheckNaN(TEXT("CharacterMovementComponent::PhysicsRotation(): DesiredRotation"));
-			const FQuat DesiredRotationQuat = DesiredRotation.Quaternion();
-			const FQuat DesiredRotationQuatAS = MovementComponent->GetOwner()->GetActorQuat().Inverse() * DesiredRotationQuat;
-			Sample.Transform.SetRotation(DesiredRotationQuatAS);
-		}
-
-		const FQuat DeltaRotation = InitialSampleTransform.GetRotation().Inverse() * Sample.Transform.GetRotation();
-
-		// since we're not extrapolating the rotation on this mode, the linear component's don't rotate
-		//Sample.LinearVelocity = DeltaRotation.RotateVector(Sample.LinearVelocity);
-		//Sample.LinearAcceleration = DeltaRotation.UnrotateVector(Sample.LinearAcceleration);
-	}
-	else
-	{
-		const FQuat DeltaRotation = FQuat(Sample.AngularVelocityAxis, Sample.AngularSpeed * IntegrationDelta);
-		const FQuat SampleRotation = Sample.Transform.GetRotation() * DeltaRotation;
-		Sample.Transform.SetRotation(SampleRotation);
-
-		Sample.LinearVelocity = DeltaRotation.RotateVector(Sample.LinearVelocity);
-		Sample.LinearAcceleration = DeltaRotation.RotateVector(Sample.LinearAcceleration);
-	}
-
+		FRotator ControlRotationDelta = ControlRotationVelocity * IntegrationDelta;
+		ControlRotationDelta = AdjustRotationVerticality(ControlRotationDelta, MovementComponent);
+		InOutControlRotationTotalDelta += ControlRotationDelta;
+
+		const FRotator TotalRotation = BaseRotation + InOutControlRotationTotalDelta;
+		const FQuat TotalRotationQuat = TotalRotation.Quaternion();
+		const FQuat TotalRotationQuatQuatAS = 
+			MovementComponent->GetOwner()->GetActorQuat().Inverse() * TotalRotationQuat;
+		
+		InOutSample.Transform.SetRotation(TotalRotationQuatQuatAS);
+		InOutSample.LinearVelocity = ControlRotationDelta.RotateVector(InOutSample.LinearVelocity);
+		InOutSample.LinearAcceleration = ControlRotationDelta.RotateVector(InOutSample.LinearAcceleration);
 	
-	const FVector Translation = Sample.LinearVelocity * IntegrationDelta;
-	Sample.Transform.AddToTranslation(Translation);
+		const FVector Translation = InOutSample.LinearVelocity * IntegrationDelta;
+		InOutSample.Transform.AddToTranslation(Translation);
+	}
 }
 
 static void PredictTrajectory(const UCharacterMovementComponent* MovementComponent
@@ -165,6 +180,7 @@ static void PredictTrajectory(const UCharacterMovementComponent* MovementCompone
 	, int32 MaxSamples
 	, const FMotionTrajectorySettings& Settings
 	, const FTrajectorySample& PresentTrajectory
+	, const FRotator& DesiredControlRotationVelocity
 	, FTrajectorySampleRange& ReturnValue)
 {
 	ReturnValue.SampleRate = SampleRate;
@@ -176,6 +192,7 @@ static void PredictTrajectory(const UCharacterMovementComponent* MovementCompone
 		FTrajectorySample PreviousSample = PresentTrajectory;
 		float AccumulatedDistance = 0.f;
 		float AccumulatedSeconds = 0.f;
+		FRotator ControlRotationTotalDelta = FRotator::ZeroRotator;
 
 		constexpr int32 DistanceDomainMask = static_cast<int32>(ETrajectorySampleDomain::Distance);
 		constexpr int32 TimeDomainMask = static_cast<int32>(ETrajectorySampleDomain::Time);
@@ -183,7 +200,12 @@ static void PredictTrajectory(const UCharacterMovementComponent* MovementCompone
 		for (int32 Step = 0; Step < MaxSamples; ++Step)
 		{
 			PreviousSample = Sample;
-			StepPrediction(IntegrationDelta, MovementComponent, Sample);
+			StepPrediction(
+				IntegrationDelta, 
+				MovementComponent, 
+				DesiredControlRotationVelocity, 
+				ControlRotationTotalDelta, 
+				Sample);
 
 			AccumulatedDistance += 
 				FVector::Distance(PreviousSample.Transform.GetLocation(), Sample.Transform.GetLocation());
@@ -218,7 +240,7 @@ static void PredictTrajectory(const UCharacterMovementComponent* MovementCompone
 	}
 }
 
-// ------------ END Derived from FCharacterMovementComponentAsyncInput::CalcVelocity() ------------ //
+// ------------ END Derived from FCharacterMovementComponentAsyncInput ------------ //
 
 FTrajectorySample UCharacterMovementTrajectoryComponent::CalcWorldSpacePresentTrajectorySample(float DeltaTime) const
 {
@@ -239,24 +261,6 @@ FTrajectorySample UCharacterMovementTrajectoryComponent::CalcWorldSpacePresentTr
 		ReturnValue.Transform = ComponentWorldTransform;
 		ReturnValue.LinearVelocity = MovementComponent->Velocity;
 		ReturnValue.LinearAcceleration = MovementComponent->GetCurrentAcceleration();
-
-		if (DeltaTime > SMALL_NUMBER)
-		{
-			const FQuat DeltaRotation = 
-				PresentTrajectorySampleWS.Transform.GetRotation().Inverse() * ComponentWorldTransform.GetRotation();
-			FVector DeltaRotationAxis;
-			float DeltaRotationAngle;
-			DeltaRotation.ToAxisAndAngle(DeltaRotationAxis, DeltaRotationAngle);
-			const float AngularSpeed = DeltaRotationAngle / DeltaTime;
-
-			ReturnValue.AngularVelocityAxis = DeltaRotationAxis;
-			ReturnValue.AngularSpeed = AngularSpeed;
-		}
-		else
-		{
-			ReturnValue.AngularVelocityAxis = PresentTrajectorySampleWS.AngularVelocityAxis;
-			ReturnValue.AngularSpeed = PresentTrajectorySampleWS.AngularSpeed;
-		}
 	}
 
 	return ReturnValue;
@@ -269,28 +273,70 @@ UCharacterMovementTrajectoryComponent::UCharacterMovementTrajectoryComponent(con
 
 void UCharacterMovementTrajectoryComponent::InitializeComponent()
 {
-	// Add a tick dependency on the character movement component since this component implementation is reliant on its internal state
 	UCharacterMovementComponent* CharacterMovementComponent = GetOwner()->FindComponentByClass<UCharacterMovementComponent>();
-	AddTickPrerequisiteComponent(CharacterMovementComponent);
-
 	if (ACharacter* Character = CharacterMovementComponent->GetCharacterOwner())
 	{
 		Character->OnCharacterMovementUpdated.AddDynamic(
 			this, 
 			&UCharacterMovementTrajectoryComponent::OnMovementUpdated);
+		PrimaryComponentTick.bCanEverTick = false;
+		PrimaryComponentTick.bStartWithTickEnabled = false;
+	}
+	else
+	{
+		UE_LOG(
+			LogMotionTrajectory, 
+			Error, 
+			TEXT("UCharacterMovementTrajectoryComponent expects the owner to be ACharacter"));
 	}
 
+
 	Super::InitializeComponent();
 }
 
 void UCharacterMovementTrajectoryComponent::UninitializeComponent()
 {
 	UCharacterMovementComponent* CharacterMovementComponent = GetOwner()->FindComponentByClass<UCharacterMovementComponent>();
-	RemoveTickPrerequisiteComponent(CharacterMovementComponent);
+	if (ACharacter* Character = CharacterMovementComponent->GetCharacterOwner())
+	{
+		Character->OnCharacterMovementUpdated.RemoveDynamic(
+			this,
+			&UCharacterMovementTrajectoryComponent::OnMovementUpdated);
+	}
 
 	Super::UninitializeComponent();
 }
 
+void UCharacterMovementTrajectoryComponent::BeginPlay()
+{
+	Super::BeginPlay();
+
+	if (const APawn* Pawn = TryGetOwnerPawn())
+	{
+		LastDesiredControlRotation = Pawn->Controller->GetDesiredRotation();
+	}
+	else
+	{
+		LastDesiredControlRotation = FRotator::ZeroRotator;
+	}
+
+	DesiredControlRotationVelocity = FRotator::ZeroRotator;
+}
+
+void UCharacterMovementTrajectoryComponent::TickTrajectory(float DeltaTime)
+{
+	Super::TickTrajectory(DeltaTime);
+	if (DeltaTime > SMALL_NUMBER)
+	{
+		if (const APawn* Pawn = TryGetOwnerPawn())
+		{
+			const FRotator CurrentDesiredRotation = Pawn->Controller->GetDesiredRotation();
+			DesiredControlRotationVelocity = (CurrentDesiredRotation - LastDesiredControlRotation) * (1.0f / DeltaTime);
+			LastDesiredControlRotation = CurrentDesiredRotation;
+		}
+	}
+}
+
 void UCharacterMovementTrajectoryComponent::OnMovementUpdated(float DeltaSeconds, FVector OldLocation, FVector OldVelocity)
 {
 	TickTrajectory(DeltaSeconds);
@@ -317,7 +363,14 @@ FTrajectorySampleRange UCharacterMovementTrajectoryComponent::GetTrajectoryWithS
 	{
 		// Step the prediction iteratively towards the specified domain horizon(s)
 		FTrajectorySampleRange Prediction(SampleRate);
-		PredictTrajectory(MovementComponent, SampleRate, MaxSamples, Settings, PresentTrajectorySampleLS, Prediction);
+		PredictTrajectory(
+			MovementComponent, 
+			SampleRate, 
+			MaxSamples, 
+			Settings, 
+			PresentTrajectorySampleLS, 
+			DesiredControlRotationVelocity,
+			Prediction);
 
 		// Combine past, present, and future into a uniformly sampled complete trajectory
 		return CombineHistoryPresentPrediction(bIncludeHistory, Prediction);
diff --git a/Engine/Plugins/Experimental/Animation/MotionTrajectory/Source/MotionTrajectory/Private/MotionTrajectoryLibrary.cpp b/Engine/Plugins/Experimental/Animation/MotionTrajectory/Source/MotionTrajectory/Private/MotionTrajectoryLibrary.cpp
--- a/Engine/Plugins/Experimental/Animation/MotionTrajectory/Source/MotionTrajectory/Private/MotionTrajectoryLibrary.cpp
+++ b/Engine/Plugins/Experimental/Animation/MotionTrajectory/Source/MotionTrajectory/Private/MotionTrajectoryLibrary.cpp
@@ -227,9 +227,9 @@ FTrajectorySampleRange UMotionTrajectoryBlueprintLibrary::MakeTrajectoryRelative
 	const AActor* Owner = Component->GetOwner();
 	const FTransform OwnerTransformWS = Owner->GetActorTransform();
 	const FTransform ComponentTransformWS = Component->GetComponentTransform();
-	const FQuat Rotation = ComponentTransformWS.GetRotation().Inverse() * OwnerTransformWS.GetRotation();
+	const FTransform ReferenceChangeTransform = OwnerTransformWS.GetRelativeTransform(ComponentTransformWS);
 
-	ActorTrajectory.Rotate(Rotation);
+	ActorTrajectory.TransformReferenceFrame(ReferenceChangeTransform);
 	return ActorTrajectory;
 }
 
@@ -238,6 +238,8 @@ void UMotionTrajectoryBlueprintLibrary::DebugDrawTrajectory(const AActor* Actor
 	, const FTrajectorySampleRange& Trajectory
 	, const FLinearColor PredictionColor
 	, const FLinearColor HistoryColor
+	, float TransformScale
+	, float TransformThickness
 	, float ArrowScale
 	, float ArrowSize
 	, float ArrowThickness
@@ -250,6 +252,8 @@ void UMotionTrajectoryBlueprintLibrary::DebugDrawTrajectory(const AActor* Actor
 			, WorldTransform.IsValid() ? WorldTransform : FTransform::Identity
 			, PredictionColor
 			, HistoryColor
+			, TransformScale
+			, TransformThickness
 			, ArrowScale
 			, ArrowSize
 			, ArrowThickness
diff --git a/Engine/Plugins/Experimental/Animation/MotionTrajectory/Source/MotionTrajectory/Public/MotionTrajectory.h b/Engine/Plugins/Experimental/Animation/MotionTrajectory/Source/MotionTrajectory/Public/MotionTrajectory.h
--- a/Engine/Plugins/Experimental/Animation/MotionTrajectory/Source/MotionTrajectory/Public/MotionTrajectory.h
+++ b/Engine/Plugins/Experimental/Animation/MotionTrajectory/Source/MotionTrajectory/Public/MotionTrajectory.h
@@ -67,14 +67,12 @@ protected:
 	// Retained trajectory history samples
 	TRingBuffer<FTrajectorySample> SampleHistory = {};
 
-	// Retained instantaneous/present trajectory
+	// Retained instantaneous/present trajectory in world space
 	FTrajectorySample PresentTrajectorySampleWS = {};
 
+	// Retained instantaneous/present trajectory in actor space
 	FTrajectorySample PresentTrajectorySampleLS = {};
 
-	// if anything makes the trajectory tick before ::TickComponent, skip it.
-	uint32 LastTrajectoryTickFrame = 0u;
-
 	// Gets the instantaneous/present trajectory sample of the current frame
 	virtual FTrajectorySample CalcWorldSpacePresentTrajectorySample(float DeltaTime) const;
 
@@ -87,8 +85,6 @@ protected:
 	// Combines all trajectory samples in the past, present, and future into a unifed trajectory range
 	FTrajectorySampleRange CombineHistoryPresentPrediction(bool bIncludeHistory, const FTrajectorySampleRange& Prediction) const;
 
-	// Checks the component internal frame counter to say if the component has already ticked this frame or not
-	bool HasTrajectoryTickedThisFrame() const;
 
 	// Forcefully evicts all trajectory history and resets internal history tracking state
 	void FlushHistory();
diff --git a/Engine/Plugins/Experimental/Animation/MotionTrajectory/Source/MotionTrajectory/Public/MotionTrajectoryCharacterMovement.h b/Engine/Plugins/Experimental/Animation/MotionTrajectory/Source/MotionTrajectory/Public/MotionTrajectoryCharacterMovement.h
--- a/Engine/Plugins/Experimental/Animation/MotionTrajectory/Source/MotionTrajectory/Public/MotionTrajectoryCharacterMovement.h
+++ b/Engine/Plugins/Experimental/Animation/MotionTrajectory/Source/MotionTrajectory/Public/MotionTrajectoryCharacterMovement.h
@@ -16,11 +16,18 @@ protected:
 
 	// Begin UMotionTrajectoryComponent Interface
 	virtual FTrajectorySample CalcWorldSpacePresentTrajectorySample(float DeltaTime) const override;
+	virtual void TickTrajectory(float DeltaTime) override;
 	// End UMotionTrajectoryComponent Interface
 
 	UFUNCTION()
 	void OnMovementUpdated(float DeltaSeconds, FVector OldLocation, FVector OldVelocity);
 
+	UPROPERTY(Transient)
+	FRotator LastDesiredControlRotation = FRotator::ZeroRotator;
+	
+	UPROPERTY(Transient)
+	FRotator DesiredControlRotationVelocity = FRotator::ZeroRotator;
+
 public:
 
 	UCharacterMovementTrajectoryComponent(const FObjectInitializer& ObjectInitializer);
@@ -28,6 +35,7 @@ public:
 	// Begin UActorComponent Interface
 	virtual void InitializeComponent() override;
 	virtual void UninitializeComponent() override;
+	virtual void BeginPlay() override;
 	// End UActorComponent Interface
 
 	// Begin UMotionTrajectoryComponent Interface
diff --git a/Engine/Plugins/Experimental/Animation/MotionTrajectory/Source/MotionTrajectory/Public/MotionTrajectoryLibrary.h b/Engine/Plugins/Experimental/Animation/MotionTrajectory/Source/MotionTrajectory/Public/MotionTrajectoryLibrary.h
--- a/Engine/Plugins/Experimental/Animation/MotionTrajectory/Source/MotionTrajectory/Public/MotionTrajectoryLibrary.h
+++ b/Engine/Plugins/Experimental/Animation/MotionTrajectory/Source/MotionTrajectory/Public/MotionTrajectoryLibrary.h
@@ -97,7 +97,9 @@ public:
 		, const FTrajectorySampleRange& Trajectory
 		, const FLinearColor PredictionColor = FLinearColor(0.f, 1.f, 0.f)
 		, const FLinearColor HistoryColor = FLinearColor(0.f, 0.f, 1.f)
-		, float ArrowScale = 10.f
+		, float TransformScale = 10.f
+		, float TransformThickness = 2.f
+		, float ArrowScale = 0.025f
 		, float ArrowSize = 40.f
 		, float ArrowThickness = 2.f
 	);
diff --git a/Engine/Source/Runtime/Engine/Private/Animation/MotionTrajectoryTypes.cpp b/Engine/Source/Runtime/Engine/Private/Animation/MotionTrajectoryTypes.cpp
--- a/Engine/Source/Runtime/Engine/Private/Animation/MotionTrajectoryTypes.cpp
+++ b/Engine/Source/Runtime/Engine/Private/Animation/MotionTrajectoryTypes.cpp
@@ -60,7 +60,6 @@ bool FTrajectorySample::IsZeroSample() const
 		&& LinearAcceleration.IsNearlyZero()
 		&& Transform.GetTranslation().IsNearlyZero()
 		&& FMath::IsNearlyZero(AccumulatedDistance)
-		&& FMath::IsNearlyZero(AngularSpeed)
 		&& Transform.GetRotation().IsIdentity();
 }
 
@@ -74,22 +73,6 @@ FTrajectorySample FTrajectorySample::Lerp(const FTrajectorySample& Sample, float
 
 	Interp.Transform.Blend(Transform, Sample.Transform, Alpha);
 	
-	// TODO: this is very simple, more like Lerp than Slerp, can we do better?
-	const FVector AngularVelocity = AngularVelocityAxis * AngularSpeed;
-	const FVector SampleAngularVelocity = Sample.AngularVelocityAxis * Sample.AngularSpeed;
-	const FVector LerpedAngularVelocity = FMath::Lerp(AngularVelocity, SampleAngularVelocity, Alpha);
-	const float LerpedAngularSpeed = LerpedAngularVelocity.Size();
-	if (LerpedAngularSpeed > SMALL_NUMBER)
-	{
-		Interp.AngularVelocityAxis = LerpedAngularVelocity / LerpedAngularSpeed;
-		Interp.AngularSpeed = LerpedAngularSpeed;
-	}
-	else
-	{
-		Interp.AngularVelocityAxis = FVector::ZeroVector;
-		Interp.AngularSpeed = 0.0f;
-	}
-
 	return Interp;
 }
 
@@ -125,24 +108,6 @@ FTrajectorySample FTrajectorySample::SmoothInterp(const FTrajectorySample& PrevS
 
 	Interp.Transform.SetRotation(FQuat::Squad(Q1, T0, Q2, T1, Alpha));
 
-	const FVector V0 = PrevSample.AngularVelocityAxis * PrevSample.AngularSpeed;
-	const FVector V1 = AngularVelocityAxis * AngularSpeed;
-	const FVector V2 = Sample.AngularVelocityAxis * Sample.AngularSpeed;
-	const FVector V3 = NextSample.AngularVelocityAxis * NextSample.AngularSpeed;
-
-	const FVector LerpedAngularVelocity = CubicCRSplineInterpSafe(V0, V1, V2, V3, Alpha);
-	const float LerpedAngularSpeed = LerpedAngularVelocity.Size();
-	if (LerpedAngularSpeed > SMALL_NUMBER)
-	{
-		Interp.AngularVelocityAxis = LerpedAngularVelocity / LerpedAngularSpeed;
-		Interp.AngularSpeed = LerpedAngularSpeed;
-	}
-	else
-	{
-		Interp.AngularVelocityAxis = FVector::ZeroVector;
-		Interp.AngularSpeed = 0.0f;
-	}
-
 	return Interp;
 }
 
@@ -167,7 +132,14 @@ void FTrajectorySample::PrependOffset(const FTransform DeltaTransform, float Del
 
 	LinearVelocity = DeltaTransform.TransformVectorNoScale(LinearVelocity);
 	LinearAcceleration = DeltaTransform.TransformVectorNoScale(LinearAcceleration);
-	AngularVelocityAxis = DeltaTransform.TransformVectorNoScale(AngularVelocityAxis);
+}
+
+void FTrajectorySample::TransformReferenceFrame(const FTransform DeltaTransform)
+{
+	Transform = DeltaTransform.Inverse() * Transform * DeltaTransform;
+
+	LinearVelocity = DeltaTransform.TransformVectorNoScale(LinearVelocity);
+	LinearAcceleration = DeltaTransform.TransformVectorNoScale(LinearAcceleration);
 }
 
 bool FTrajectorySampleRange::HasSamples() const
@@ -199,14 +171,24 @@ void FTrajectorySampleRange::Rotate(const FQuat& Rotation)
 	}
 }
 
+void FTrajectorySampleRange::TransformReferenceFrame(const FTransform& Transform)
+{
+	for (auto& Sample : Samples)
+	{
+		Sample.TransformReferenceFrame(Transform);
+	}
+}
+
 void FTrajectorySampleRange::DebugDrawTrajectory(bool bEnable
 	, const UWorld* World
 	, const FTransform& WorldTransform
 	, const FLinearColor PredictionColor
 	, const FLinearColor HistoryColor
-	, float ArrowScale
-	, float ArrowSize
-	, float ArrowThickness) const
+	, float TransformScale
+	, float TransformThickness
+	, float VelArrowScale
+	, float VelArrowSize
+	, float VelArrowThickness) const
 {
 	if (bEnable
 #if ENABLE_ANIM_DEBUG
@@ -222,17 +204,37 @@ void FTrajectorySampleRange::DebugDrawTrajectory(bool bEnable
 #endif
 			for (int32 Idx = 0, Num = Samples.Num(); Idx < Num; Idx++)
 			{
-				const FVector WorldPosition = WorldTransform.TransformPosition(Samples[Idx].Transform.GetTranslation());
-				const FVector WorldForward = 
-					(WorldTransform.TransformVectorNoScale(Samples[Idx].Transform.GetRotation().GetAxisX()) *
-					 ArrowScale) 
-					+ WorldPosition;
+				const FTransform SampleTransformWS = Samples[Idx].Transform * WorldTransform;
+				const FVector SamplePositionWS = SampleTransformWS.GetTranslation();
+
+				const FVector WorldVelocity =
+					(WorldTransform.TransformVector(Samples[Idx].LinearVelocity) * VelArrowScale) + SamplePositionWS;
 
 				// Interpolate the history and prediction color over the entire trajectory range
 				const float ColorLerp = static_cast<float>(Idx) / static_cast<float>(Num);
 				const FLinearColor Color = FLinearColor::LerpUsingHSV(PredictionColor, HistoryColor, ColorLerp);
 
-				DrawDebugDirectionalArrow(World, WorldPosition, WorldForward, ArrowSize, Color.ToFColor(true), false, 0.f, 0, ArrowThickness);
+				if (VelArrowScale > 0.0f)
+				{
+					DrawDebugDirectionalArrow(
+						World, SamplePositionWS, WorldVelocity, 
+						VelArrowSize, Color.ToFColor(true), false, 0.f, 0, VelArrowThickness);
+				}
+
+				if (TransformScale > 0.0f)
+				{
+					DrawDebugCoordinateSystem(
+						World, SampleTransformWS.GetLocation(), SampleTransformWS.Rotator(),
+						TransformScale, false, 0.f, 0, TransformThickness);
+
+					if (VelArrowThickness == 0.0f)
+					{
+						DrawDebugSphere(
+							World, SamplePositionWS,										
+							TransformScale * 0.5f, 4, Color.ToFColor(true), false, 0.f, 0, TransformThickness * 0.5f);
+					}
+				}
+
 #if ENABLE_ANIM_DEBUG
 				FString DebugString;
 				FString DebugSampleString;
@@ -275,7 +277,7 @@ void FTrajectorySampleRange::DebugDrawTrajectory(bool bEnable
 						DrawDebugString(World, WorldTransform.GetLocation() + DebugSampleTypeOffset, DebugString, nullptr, FColor::White, 0.f, false, 1.f);
 					}
 
-					DrawDebugString(World, WorldForward + DebugSampleOffset, DebugSampleString, nullptr, FColor::White, 0.f, false, 1.f);
+					DrawDebugString(World, WorldVelocity + DebugSampleOffset, DebugSampleString, nullptr, FColor::White, 0.f, false, 1.f);
 				}
 #endif
 			}
diff --git a/Engine/Source/Runtime/Engine/Public/Animation/MotionTrajectoryTypes.h b/Engine/Source/Runtime/Engine/Public/Animation/MotionTrajectoryTypes.h
--- a/Engine/Source/Runtime/Engine/Public/Animation/MotionTrajectoryTypes.h
+++ b/Engine/Source/Runtime/Engine/Public/Animation/MotionTrajectoryTypes.h
@@ -40,14 +40,6 @@ struct ENGINE_API FTrajectorySample
 	UPROPERTY(EditAnywhere, BlueprintReadWrite, Category="Motion Trajectory")
 	FVector LinearVelocity = FVector::ZeroVector;
 
-	// Axis around which the angular velocity rotates the in-motion object
-	UPROPERTY(EditAnywhere, BlueprintReadWrite, Category = "Motion Trajectory")
-	FVector AngularVelocityAxis = FVector::ZeroVector;
-
-	// Rotation speed around AngularVelocityAxis
-	UPROPERTY(EditAnywhere, BlueprintReadWrite, Category = "Motion Trajectory")
-	float AngularSpeed = 0.0f;
-
 	// Linear acceleration relative to the sampled in-motion object
 	UPROPERTY(EditAnywhere, BlueprintReadWrite, Category="Motion Trajectory")
 	FVector LinearAcceleration = FVector::ZeroVector;
@@ -65,6 +57,8 @@ struct ENGINE_API FTrajectorySample
 	// DeltaSeconds
 	void PrependOffset(const FTransform DeltaTransform, float DeltaSeconds);
 
+	void TransformReferenceFrame(const FTransform DeltaTransform);
+
 	// Determines if all sample properties are zeroed
 	bool IsZeroSample() const;
 };
@@ -78,7 +72,9 @@ struct ENGINE_API FTrajectorySampleRange
 	// Debug rendering contants
 	static constexpr FLinearColor DebugDefaultPredictionColor{ 0.f, 1.f, 0.f };
 	static constexpr FLinearColor DebugDefaultHistoryColor{ 0.f, 0.f, 1.f };
-	static constexpr float DebugDefaultArrowScale = 10.f;
+	static constexpr float DebugDefaultTransformScale = 10.f;
+	static constexpr float DebugDefaultTransformThickness = 2.f;
+	static constexpr float DebugDefaultArrowScale = 0.025f;
 	static constexpr float DebugDefaultArrowSize = 40.f;
 	static constexpr float DebugDefaultArrowThickness = 2.f;
 
@@ -101,6 +97,9 @@ struct ENGINE_API FTrajectorySampleRange
 	// Rotates all samples in the trajectory
 	void Rotate(const FQuat& Rotation);
 
+	// Rotates all samples in the trajectory
+	void TransformReferenceFrame(const FTransform& Transform);
+
 	// Determine if any trajectory samples are present
 	bool HasSamples() const;
 
@@ -113,9 +112,11 @@ struct ENGINE_API FTrajectorySampleRange
 		, const FTransform& WorldTransform
 		, const FLinearColor PredictionColor = DebugDefaultPredictionColor
 		, const FLinearColor HistoryColor = DebugDefaultHistoryColor
-		, float ArrowScale = DebugDefaultArrowScale
-		, float ArrowSize = DebugDefaultArrowSize
-		, float ArrowThickness = DebugDefaultArrowThickness) const;
+		, float TransformScale = DebugDefaultTransformScale
+		, float TransformThickness = DebugDefaultTransformThickness
+		, float VelArrowScale = DebugDefaultArrowScale
+		, float VelArrowSize = DebugDefaultArrowSize
+		, float VelArrowThickness = DebugDefaultArrowThickness) const;
 
 	// Iterator for precise subsampling of the trajectory
 	template <typename Container> static FTrajectorySample IterSampleTrajectory(const Container& Samples, ETrajectorySampleDomain DomainType, float DomainValue, int32& InitialIdx, bool bSmoothInterp = false)
