Fixed non-unity build error

diff --git a/Engine/Plugins/Experimental/Animation/MotionTrajectory/Source/MotionTrajectory/Private/MotionTrajectoryCharacterMovement.cpp b/Engine/Plugins/Experimental/Animation/MotionTrajectory/Source/MotionTrajectory/Private/MotionTrajectoryCharacterMovement.cpp
--- a/Engine/Plugins/Experimental/Animation/MotionTrajectory/Source/MotionTrajectory/Private/MotionTrajectoryCharacterMovement.cpp
+++ b/Engine/Plugins/Experimental/Animation/MotionTrajectory/Source/MotionTrajectory/Private/MotionTrajectoryCharacterMovement.cpp
@@ -6,6 +6,7 @@
 #include "GameFramework/CharacterMovementComponent.h"
 #include "GameFramework/Pawn.h"
 #include "GameFramework/Character.h"
+#include "GameFramework/Controller.h"
 #include "KismetAnimationLibrary.h"
 
 // ----------- BEGIN Derived from FCharacterMovementComponentAsyncInput ----------- //
