MotionMatchingFFDemo - Updated trajectory component to match new movement model.
- Since the new movement component extends from UCharacterMovementComponent, I updated UMotionTrajectoryCharacterMovement to create an interface that supports the modifications
- Converted a lot of constants in the original model with function calls so the values can change dynamically.
- Implemented a new trajectory component which inherits from UMotionTrajectoryCharacterMovement. Matching the new movement model required only a straightforward implementation of the new virtual functions.
- One aspect that isn't being properly predicted is camera assist. In the future we should evaluate if we should model it more accurately or, alternatively, ignore it altogether, if it proves to produce bad predictions.

diff --git a/Engine/Plugins/Experimental/Animation/MotionTrajectory/Source/MotionTrajectory/Private/MotionTrajectoryCharacterMovement.cpp b/Engine/Plugins/Experimental/Animation/MotionTrajectory/Source/MotionTrajectory/Private/MotionTrajectoryCharacterMovement.cpp
--- a/Engine/Plugins/Experimental/Animation/MotionTrajectory/Source/MotionTrajectory/Private/MotionTrajectoryCharacterMovement.cpp
+++ b/Engine/Plugins/Experimental/Animation/MotionTrajectory/Source/MotionTrajectory/Private/MotionTrajectoryCharacterMovement.cpp
@@ -9,12 +9,213 @@
 #include "GameFramework/Controller.h"
 #include "KismetAnimationLibrary.h"
 
+FTrajectorySample UCharacterMovementTrajectoryComponent::CalcWorldSpacePresentTrajectorySample(float DeltaTime) const
+{
+	FTrajectorySample ReturnValue;
+	
+	const APawn* Pawn = TryGetOwnerPawn();
+	if (!Pawn)
+	{
+		return ReturnValue;
+	}
+
+	const UCharacterMovementComponent* MovementComponent = Cast<UCharacterMovementComponent>(Pawn->GetMovementComponent());
+	if (ensure(MovementComponent) && MovementComponent->MovementMode == EMovementMode::MOVE_Walking)
+	{
+		FTransform ComponentWorldTransform = Pawn->GetActorTransform();
+		FVector MovementComponentPosition = MovementComponent->GetLastUpdateLocation();
+
+		ReturnValue.Transform = ComponentWorldTransform;
+		ReturnValue.LinearVelocity = MovementComponent->Velocity;
+	}
+
+	return ReturnValue;
+}
+
+UCharacterMovementTrajectoryComponent::UCharacterMovementTrajectoryComponent(const FObjectInitializer& ObjectInitializer) : Super(ObjectInitializer)
+{
+	bWantsInitializeComponent = true;
+}
+
+void UCharacterMovementTrajectoryComponent::InitializeComponent()
+{
+	if (ACharacter* Character = Cast<ACharacter>(GetOwner()))
+	{
+		Character->OnCharacterMovementUpdated.AddDynamic(
+			this, 
+			&UCharacterMovementTrajectoryComponent::OnMovementUpdated);
+		PrimaryComponentTick.bCanEverTick = false;
+		PrimaryComponentTick.bStartWithTickEnabled = false;
+	}
+	else
+	{
+		UE_LOG(
+			LogMotionTrajectory, 
+			Error, 
+			TEXT("UCharacterMovementTrajectoryComponent expects the owner to be ACharacter"));
+	}
+
+
+	Super::InitializeComponent();
+}
+
+void UCharacterMovementTrajectoryComponent::UninitializeComponent()
+{
+	UCharacterMovementComponent* CharacterMovementComponent = GetOwner()->FindComponentByClass<UCharacterMovementComponent>();
+	if (ACharacter* Character = CharacterMovementComponent->GetCharacterOwner())
+	{
+		Character->OnCharacterMovementUpdated.RemoveDynamic(
+			this,
+			&UCharacterMovementTrajectoryComponent::OnMovementUpdated);
+	}
+
+	Super::UninitializeComponent();
+}
+
+void UCharacterMovementTrajectoryComponent::BeginPlay()
+{
+	Super::BeginPlay();
+
+	if (const APawn* Pawn = TryGetOwnerPawn())
+	{
+		LastDesiredControlRotation = Pawn->Controller->GetDesiredRotation();
+	}
+	else
+	{
+		LastDesiredControlRotation = FRotator::ZeroRotator;
+	}
+
+	DesiredControlRotationVelocity = FRotator::ZeroRotator;
+}
+
+void UCharacterMovementTrajectoryComponent::TickTrajectory(float DeltaTime)
+{
+	Super::TickTrajectory(DeltaTime);
+	if (DeltaTime > SMALL_NUMBER)
+	{
+		if (const APawn* Pawn = TryGetOwnerPawn())
+		{
+			const FRotator CurrentDesiredRotation = Pawn->Controller->GetDesiredRotation();
+			DesiredControlRotationVelocity = (CurrentDesiredRotation - LastDesiredControlRotation) * (1.0f / DeltaTime);
+			LastDesiredControlRotation = CurrentDesiredRotation;
+		}
+	}
+}
+
+void UCharacterMovementTrajectoryComponent::OnMovementUpdated(float DeltaSeconds, FVector OldLocation, FVector OldVelocity)
+{
+	TickTrajectory(DeltaSeconds);
+}
+
+FTrajectorySampleRange UCharacterMovementTrajectoryComponent::GetTrajectory() const
+{
+	return GetTrajectoryWithSettings(PredictionSettings, bPredictionIncludesHistory);
+}
+
+FTrajectorySampleRange UCharacterMovementTrajectoryComponent::GetTrajectoryWithSettings(const FMotionTrajectorySettings& Settings
+	, bool bIncludeHistory) const
+{
+	const APawn* Pawn = TryGetOwnerPawn();
+	if (!Pawn)
+	{
+		return FTrajectorySampleRange(SampleRate);
+	}
+
+	const UCharacterMovementComponent* MovementComponent = Cast<UCharacterMovementComponent>(Pawn->GetMovementComponent());
+
+	// Currently the trajectory prediction only supports the walking movement mode of the character movement component
+	if (ensure(MovementComponent) && MovementComponent->MovementMode == EMovementMode::MOVE_Walking)
+	{
+		// Step the prediction iteratively towards the specified domain horizon(s)
+		FTrajectorySampleRange Prediction(SampleRate);
+		PredictTrajectory(
+			SampleRate, 
+			MaxSamples, 
+			Settings, 
+			PresentTrajectorySampleLS, 
+			DesiredControlRotationVelocity,
+			Prediction);
+
+		// Combine past, present, and future into a uniformly sampled complete trajectory
+		return CombineHistoryPresentPrediction(bIncludeHistory, Prediction);
+	}
+	else
+	{
+		return FTrajectorySampleRange(SampleRate);
+	}
+}
+
+void UCharacterMovementTrajectoryComponent::PredictTrajectory(
+	int32 InSampleRate,
+	int32 InMaxSamples,
+	const FMotionTrajectorySettings& Settings,
+	const FTrajectorySample& PresentTrajectory,
+	const FRotator& InDesiredControlRotationVelocity,
+	FTrajectorySampleRange& OutTrajectoryRange) const
+{
+	OutTrajectoryRange.SampleRate = InSampleRate;
+	const float IntegrationDelta = 1.f / static_cast<float>(InSampleRate);
+
+	if (!!Settings.Domain)
+	{
+		FTrajectorySample Sample = PresentTrajectory;
+		FTrajectorySample PreviousSample = PresentTrajectory;
+		float AccumulatedDistance = 0.f;
+		float AccumulatedSeconds = 0.f;
+		FRotator ControlRotationTotalDelta = FRotator::ZeroRotator;
+
+		constexpr int32 DistanceDomainMask = static_cast<int32>(ETrajectorySampleDomain::Distance);
+		constexpr int32 TimeDomainMask = static_cast<int32>(ETrajectorySampleDomain::Time);
+
+		for (int32 Step = 0; Step < InMaxSamples; ++Step)
+		{
+			PreviousSample = Sample;
+			StepPrediction(
+				IntegrationDelta, 
+				InDesiredControlRotationVelocity, 
+				ControlRotationTotalDelta, 
+				Sample);
+
+			AccumulatedDistance += 
+				FVector::Distance(PreviousSample.Transform.GetLocation(), Sample.Transform.GetLocation());
+			Sample.AccumulatedDistance = AccumulatedDistance;
+			AccumulatedSeconds += IntegrationDelta;
+			Sample.AccumulatedSeconds = AccumulatedSeconds;
+
+			if (FMath::IsNearlyEqual(FMath::Abs(Sample.AccumulatedDistance - PreviousSample.AccumulatedDistance), SMALL_NUMBER))
+			{
+				break;
+			}
+
+			OutTrajectoryRange.Samples.Add(Sample);
+
+			if (((Settings.Domain & DistanceDomainMask) == DistanceDomainMask)
+				&& (Settings.Distance > 0.f)
+				&& (Sample.AccumulatedDistance < Settings.Distance))
+			{
+				continue;
+			}
+
+			if (((Settings.Domain & TimeDomainMask) == TimeDomainMask)
+				&& (Settings.Seconds > 0.f)
+				&& (Step * IntegrationDelta < Settings.Seconds)
+				&& (Sample.AccumulatedDistance > PreviousSample.AccumulatedDistance))
+			{
+				continue;
+			}
+
+			break;
+		}
+	}
+}
+
+
 // ----------- BEGIN Derived from FCharacterMovementComponentAsyncInput ----------- //
 // FCharacterMovementComponentAsyncInput::CalcVelocity() for linear motion
 // FCharacterMovementComponentAsyncInput::PhysicsRotation() for angular motion
 
 static FRotator AdjustRotationVerticality(
-	const FRotator& InRotation, 
+	const FRotator& InRotation,
 	const UCharacterMovementComponent* MovementComponent)
 {
 	FRotator OutRotation = InRotation;
@@ -71,24 +272,35 @@
 	return OutRotation;
 }
 
-static void StepPrediction(float IntegrationDelta
-	, const UCharacterMovementComponent* MovementComponent
-	, const FRotator& ControlRotationVelocity
-	, FRotator& InOutControlRotationTotalDelta
-	, FTrajectorySample& InOutSample)
+void UCharacterMovementTrajectoryComponent::StepPrediction(
+	float IntegrationDelta,
+	const FRotator& ControlRotationVelocity,
+	FRotator& InOutControlRotationTotalDelta,
+	FTrajectorySample& InOutSample) const
 {
 	const FTransform InitialSampleTransform = InOutSample.Transform;
 
+	const UCharacterMovementComponent* MovementComponent = 
+		Cast<UCharacterMovementComponent>(TryGetOwnerPawn()->GetMovementComponent());
+	if (!IsValid(MovementComponent))
+	{
+		return;
+	}
+
 	if (MovementComponent->GetCurrentAcceleration().IsZero())
 	{
-		float ActualBrakingFriction = (MovementComponent->bUseSeparateBrakingFriction ? MovementComponent->BrakingFriction : MovementComponent->GroundFriction);
+		float ActualBrakingFriction = 
+			MovementComponent->bUseSeparateBrakingFriction ? 
+			MovementComponent->BrakingFriction : 
+			GetFriction(MovementComponent, InOutSample, IntegrationDelta);
 
 		if (!InOutSample.LinearVelocity.IsZero())
 		{
 			const float FrictionFactor = FMath::Max(0.f, MovementComponent->BrakingFrictionFactor);
 			ActualBrakingFriction = FMath::Max(0.f, ActualBrakingFriction * FrictionFactor);
 				
-			const float BrakingDeceleration = FMath::Max(0.f, MovementComponent->GetMaxBrakingDeceleration());
+			const float BrakingDeceleration = 
+				FMath::Max(0.f, GetMaxBrakingDeceleration(MovementComponent, InOutSample, IntegrationDelta));
 			const bool bZeroFriction = (ActualBrakingFriction == 0.f);
 			const bool bZeroBraking = (BrakingDeceleration == 0.f);
 
@@ -127,36 +339,32 @@
 	}
 	else
 	{
-		const FVector AccelDir = InOutSample.LinearAcceleration.GetSafeNormal();
+		const FVector LinearAccelerationWS = GetAccelerationWS(MovementComponent, InOutSample, IntegrationDelta);
+		const FVector LinearAccelerationAS = 
+			TryGetOwnerPawn()->GetActorTransform().Inverse().TransformVectorNoScale(LinearAccelerationWS);
+		const FVector RotatedLinearAccelerationAS = InOutControlRotationTotalDelta.RotateVector(LinearAccelerationAS);
+
+		const FVector AccelDir = RotatedLinearAccelerationAS.GetSafeNormal();
 		const float VelSize = InOutSample.LinearVelocity.Size();
-		InOutSample.LinearVelocity = InOutSample.LinearVelocity - (InOutSample.LinearVelocity - AccelDir * VelSize) * FMath::Min(IntegrationDelta * MovementComponent->GroundFriction, 1.f);
+		InOutSample.LinearVelocity = 
+			InOutSample.LinearVelocity - 
+			(InOutSample.LinearVelocity - AccelDir * VelSize) * 
+			FMath::Min(IntegrationDelta * MovementComponent->GroundFriction, 1.f);
 
-		const float MaxInputSpeed = FMath::Max(MovementComponent->GetMaxSpeed() * MovementComponent->GetAnalogInputModifier(), MovementComponent->GetMinAnalogSpeed());
-		InOutSample.LinearVelocity += InOutSample.LinearAcceleration * IntegrationDelta;
+		const float MaxInputSpeed = FMath::Max(
+			MovementComponent->GetMaxSpeed() * MovementComponent->GetAnalogInputModifier(), 
+			MovementComponent->GetMinAnalogSpeed());
+		InOutSample.LinearVelocity += RotatedLinearAccelerationAS * IntegrationDelta;
 		InOutSample.LinearVelocity = InOutSample.LinearVelocity.GetClampedToMaxSize(MaxInputSpeed);
 	}
 
 	const FTransform PrevTransformWS = InitialSampleTransform * MovementComponent->GetOwner()->GetActorTransform();
 
 	{
-		FRotator DeltaRot = MovementComponent->GetDeltaRotation(IntegrationDelta);
-		DeltaRot.DiagnosticCheckNaN(TEXT("StepPrediction: MovementComponent->GetDeltaRotation"));
+		const FRotator PrevRotator = PrevTransformWS.GetRotation().Rotator() - InOutControlRotationTotalDelta;
+		PrevRotator.DiagnosticCheckNaN(TEXT("StepPrediction: PrevRotator"));
 
-		FRotator BaseRotation = FRotator::ZeroRotator; 
-		if (MovementComponent->bOrientRotationToMovement)
-		{
-			const FRotator PrevRotator = PrevTransformWS.GetRotation().Rotator() - InOutControlRotationTotalDelta;
-			PrevRotator.DiagnosticCheckNaN(TEXT("StepPrediction: PrevRotator"));
+		FRotator BaseRotation = GetBaseRotationWS(MovementComponent, InOutSample, PrevRotator, IntegrationDelta);
 
-			BaseRotation =
-				MovementComponent->ComputeOrientToMovementRotation(PrevRotator, IntegrationDelta, DeltaRot);
-			BaseRotation = AdjustRotationVerticality(BaseRotation, MovementComponent);
-			BaseRotation = LimitRotationRate(BaseRotation, PrevRotator, DeltaRot, MovementComponent);
-		}
-		else
-		{
-			BaseRotation = MovementComponent->GetOwner()->GetActorTransform().Rotator();
-		}
-
 		FRotator ControlRotationDelta = ControlRotationVelocity * IntegrationDelta;
 		ControlRotationDelta = AdjustRotationVerticality(ControlRotationDelta, MovementComponent);
@@ -169,215 +377,59 @@
 		
 		InOutSample.Transform.SetRotation(TotalRotationQuatQuatAS);
 		InOutSample.LinearVelocity = ControlRotationDelta.RotateVector(InOutSample.LinearVelocity);
-		InOutSample.LinearAcceleration = ControlRotationDelta.RotateVector(InOutSample.LinearAcceleration);
 	
 		const FVector Translation = InOutSample.LinearVelocity * IntegrationDelta;
 		InOutSample.Transform.AddToTranslation(Translation);
 	}
 }
 
-static void PredictTrajectory(const UCharacterMovementComponent* MovementComponent
-	, int32 SampleRate
-	, int32 MaxSamples
-	, const FMotionTrajectorySettings& Settings
-	, const FTrajectorySample& PresentTrajectory
-	, const FRotator& DesiredControlRotationVelocity
-	, FTrajectorySampleRange& ReturnValue)
+float UCharacterMovementTrajectoryComponent::GetFriction(
+	const UCharacterMovementComponent* MoveComponent, 
+	const FTrajectorySample& Sample, 
+	float DeltaSeconds) const
 {
-	ReturnValue.SampleRate = SampleRate;
-	const float IntegrationDelta = 1.f / static_cast<float>(SampleRate);
-
-	if (!!Settings.Domain)
-	{
-		FTrajectorySample Sample = PresentTrajectory;
-		FTrajectorySample PreviousSample = PresentTrajectory;
-		float AccumulatedDistance = 0.f;
-		float AccumulatedSeconds = 0.f;
-		FRotator ControlRotationTotalDelta = FRotator::ZeroRotator;
-
-		constexpr int32 DistanceDomainMask = static_cast<int32>(ETrajectorySampleDomain::Distance);
-		constexpr int32 TimeDomainMask = static_cast<int32>(ETrajectorySampleDomain::Time);
-
-		for (int32 Step = 0; Step < MaxSamples; ++Step)
-		{
-			PreviousSample = Sample;
-			StepPrediction(
-				IntegrationDelta, 
-				MovementComponent, 
-				DesiredControlRotationVelocity, 
-				ControlRotationTotalDelta, 
-				Sample);
-
-			AccumulatedDistance += 
-				FVector::Distance(PreviousSample.Transform.GetLocation(), Sample.Transform.GetLocation());
-			Sample.AccumulatedDistance = AccumulatedDistance;
-			AccumulatedSeconds += IntegrationDelta;
-			Sample.AccumulatedSeconds = AccumulatedSeconds;
-
-			if (FMath::IsNearlyEqual(FMath::Abs(Sample.AccumulatedDistance - PreviousSample.AccumulatedDistance), SMALL_NUMBER))
-			{
-				break;
-			}
-
-			ReturnValue.Samples.Add(Sample);
-
-			if (((Settings.Domain & DistanceDomainMask) == DistanceDomainMask)
-				&& (Settings.Distance > 0.f)
-				&& (Sample.AccumulatedDistance < Settings.Distance))
-			{
-				continue;
-			}
-
-			if (((Settings.Domain & TimeDomainMask) == TimeDomainMask)
-				&& (Settings.Seconds > 0.f)
-				&& (Step * IntegrationDelta < Settings.Seconds)
-				&& (Sample.AccumulatedDistance > PreviousSample.AccumulatedDistance))
-			{
-				continue;
-			}
-
-			break;
-		}
-	}
+	return MoveComponent->GroundFriction;
 }
 
+float UCharacterMovementTrajectoryComponent::GetMaxBrakingDeceleration(
+	const UCharacterMovementComponent* MoveComponent,
+	const FTrajectorySample& Sample, 
+	float DeltaSeconds) const
+{
+	return MoveComponent->GetMaxBrakingDeceleration();
+}
+
+FVector UCharacterMovementTrajectoryComponent::GetAccelerationWS(
+	const UCharacterMovementComponent* MoveComponent,
+	const FTrajectorySample& Sample, 
+	float DeltaSeconds) const
+{
+	return MoveComponent->GetCurrentAcceleration();
+}
+
+FRotator UCharacterMovementTrajectoryComponent::GetBaseRotationWS(
+	const UCharacterMovementComponent* MoveComponent,
+	const FTrajectorySample& Sample,
+	const FRotator& SampleBaseRotationWS,
+	float DeltaSeconds) const
+{
+	FRotator BaseRotation = FRotator::ZeroRotator;
+	if (MoveComponent->bOrientRotationToMovement)
+	{
+		FRotator DeltaRot = MoveComponent->GetDeltaRotation(DeltaSeconds);
+		DeltaRot.DiagnosticCheckNaN(TEXT("StepPrediction: MovementComponent->GetDeltaRotation"));
+
+		BaseRotation =
+			MoveComponent->ComputeOrientToMovementRotation(SampleBaseRotationWS, DeltaSeconds, DeltaRot);
+		BaseRotation = AdjustRotationVerticality(BaseRotation, MoveComponent);
+		BaseRotation = LimitRotationRate(BaseRotation, SampleBaseRotationWS, DeltaRot, MoveComponent);
+	}
+	else
+	{
+		BaseRotation = MoveComponent->GetOwner()->GetActorTransform().Rotator();
+	}
+
+	return BaseRotation;
+}
+
 // ------------ END Derived from FCharacterMovementComponentAsyncInput ------------ //
-
-FTrajectorySample UCharacterMovementTrajectoryComponent::CalcWorldSpacePresentTrajectorySample(float DeltaTime) const
-{
-	FTrajectorySample ReturnValue;
-	
-	const APawn* Pawn = TryGetOwnerPawn();
-	if (!Pawn)
-	{
-		return ReturnValue;
-	}
-
-	const UCharacterMovementComponent* MovementComponent = Cast<UCharacterMovementComponent>(Pawn->GetMovementComponent());
-	if (ensure(MovementComponent) && MovementComponent->MovementMode == EMovementMode::MOVE_Walking)
-	{
-		FTransform ComponentWorldTransform = Pawn->GetActorTransform();
-		FVector MovementComponentPosition = MovementComponent->GetLastUpdateLocation();
-
-		ReturnValue.Transform = ComponentWorldTransform;
-		ReturnValue.LinearVelocity = MovementComponent->Velocity;
-		ReturnValue.LinearAcceleration = MovementComponent->GetCurrentAcceleration();
-	}
-
-	return ReturnValue;
-}
-
-UCharacterMovementTrajectoryComponent::UCharacterMovementTrajectoryComponent(const FObjectInitializer& ObjectInitializer) : Super(ObjectInitializer)
-{
-	bWantsInitializeComponent = true;
-}
-
-void UCharacterMovementTrajectoryComponent::InitializeComponent()
-{
-	UCharacterMovementComponent* CharacterMovementComponent = GetOwner()->FindComponentByClass<UCharacterMovementComponent>();
-	if (ACharacter* Character = CharacterMovementComponent->GetCharacterOwner())
-	{
-		Character->OnCharacterMovementUpdated.AddDynamic(
-			this, 
-			&UCharacterMovementTrajectoryComponent::OnMovementUpdated);
-		PrimaryComponentTick.bCanEverTick = false;
-		PrimaryComponentTick.bStartWithTickEnabled = false;
-	}
-	else
-	{
-		UE_LOG(
-			LogMotionTrajectory, 
-			Error, 
-			TEXT("UCharacterMovementTrajectoryComponent expects the owner to be ACharacter"));
-	}
-
-
-	Super::InitializeComponent();
-}
-
-void UCharacterMovementTrajectoryComponent::UninitializeComponent()
-{
-	UCharacterMovementComponent* CharacterMovementComponent = GetOwner()->FindComponentByClass<UCharacterMovementComponent>();
-	if (ACharacter* Character = CharacterMovementComponent->GetCharacterOwner())
-	{
-		Character->OnCharacterMovementUpdated.RemoveDynamic(
-			this,
-			&UCharacterMovementTrajectoryComponent::OnMovementUpdated);
-	}
-
-	Super::UninitializeComponent();
-}
-
-void UCharacterMovementTrajectoryComponent::BeginPlay()
-{
-	Super::BeginPlay();
-
-	if (const APawn* Pawn = TryGetOwnerPawn())
-	{
-		LastDesiredControlRotation = Pawn->Controller->GetDesiredRotation();
-	}
-	else
-	{
-		LastDesiredControlRotation = FRotator::ZeroRotator;
-	}
-
-	DesiredControlRotationVelocity = FRotator::ZeroRotator;
-}
-
-void UCharacterMovementTrajectoryComponent::TickTrajectory(float DeltaTime)
-{
-	Super::TickTrajectory(DeltaTime);
-	if (DeltaTime > SMALL_NUMBER)
-	{
-		if (const APawn* Pawn = TryGetOwnerPawn())
-		{
-			const FRotator CurrentDesiredRotation = Pawn->Controller->GetDesiredRotation();
-			DesiredControlRotationVelocity = (CurrentDesiredRotation - LastDesiredControlRotation) * (1.0f / DeltaTime);
-			LastDesiredControlRotation = CurrentDesiredRotation;
-		}
-	}
-}
-
-void UCharacterMovementTrajectoryComponent::OnMovementUpdated(float DeltaSeconds, FVector OldLocation, FVector OldVelocity)
-{
-	TickTrajectory(DeltaSeconds);
-}
-
-FTrajectorySampleRange UCharacterMovementTrajectoryComponent::GetTrajectory() const
-{
-	return GetTrajectoryWithSettings(PredictionSettings, bPredictionIncludesHistory);
-}
-
-FTrajectorySampleRange UCharacterMovementTrajectoryComponent::GetTrajectoryWithSettings(const FMotionTrajectorySettings& Settings
-	, bool bIncludeHistory) const
-{
-	const APawn* Pawn = TryGetOwnerPawn();
-	if (!Pawn)
-	{
-		return FTrajectorySampleRange(SampleRate);
-	}
-
-	const UCharacterMovementComponent* MovementComponent = Cast<UCharacterMovementComponent>(Pawn->GetMovementComponent());
-
-	// Currently the trajectory prediction only supports the walking movement mode of the character movement component
-	if (ensure(MovementComponent) && MovementComponent->MovementMode == EMovementMode::MOVE_Walking)
-	{
-		// Step the prediction iteratively towards the specified domain horizon(s)
-		FTrajectorySampleRange Prediction(SampleRate);
-		PredictTrajectory(
-			MovementComponent, 
-			SampleRate, 
-			MaxSamples, 
-			Settings, 
-			PresentTrajectorySampleLS, 
-			DesiredControlRotationVelocity,
-			Prediction);
-
-		// Combine past, present, and future into a uniformly sampled complete trajectory
-		return CombineHistoryPresentPrediction(bIncludeHistory, Prediction);
-	}
-	else
-	{
-		return FTrajectorySampleRange(SampleRate);
-	}
-}
diff --git a/Engine/Plugins/Experimental/Animation/MotionTrajectory/Source/MotionTrajectory/Private/MotionTrajectoryLibrary.cpp b/Engine/Plugins/Experimental/Animation/MotionTrajectory/Source/MotionTrajectory/Private/MotionTrajectoryLibrary.cpp
--- a/Engine/Plugins/Experimental/Animation/MotionTrajectory/Source/MotionTrajectory/Private/MotionTrajectoryLibrary.cpp
+++ b/Engine/Plugins/Experimental/Animation/MotionTrajectory/Source/MotionTrajectory/Private/MotionTrajectoryLibrary.cpp
@@ -66,13 +66,6 @@ FTrajectorySampleRange UMotionTrajectoryBlueprintLibrary::FlattenTrajectory2D(FT
 			const float VelMagnitude = bPreserveSpeed ? Sample.LinearVelocity.Size() : Sample.LinearVelocity.Size2D();
 			Sample.LinearVelocity = VelMagnitude * Sample.LinearVelocity.GetSafeNormal2D();
 		}
-
-		// Align linear acceleration via projection onto velocity
-		if (!Sample.LinearVelocity.IsZero() && !Sample.LinearAcceleration.IsZero())
-		{
-			const float AccelMagnitude = bPreserveSpeed ? Sample.LinearAcceleration.Size() : Sample.LinearAcceleration.Size2D();
-			Sample.LinearAcceleration = AccelMagnitude * Sample.LinearAcceleration.ProjectOnTo(Sample.LinearVelocity).GetSafeNormal();
-		}
 	}
 
 	// The present position sample is used as the basis for recomputing the future and history Accumulated Distance
@@ -182,12 +175,6 @@ FTrajectorySampleRange UMotionTrajectoryBlueprintLibrary::ClampTrajectoryDirecti
 				Sample.LinearVelocity = Sample.LinearVelocity.Size() * Sample.LinearVelocity.ProjectOnTo(VelocityBasis).GetSafeNormal();
 			}
 
-			// Align linear acceleration through projection onto the modified velocity
-			if (!Sample.LinearVelocity.IsZero() && !Sample.LinearAcceleration.IsZero())
-			{
-				Sample.LinearAcceleration = Sample.LinearAcceleration.Size() * Sample.LinearAcceleration.ProjectOnTo(Sample.LinearVelocity).GetSafeNormal();
-			}
-
 			// Align the position path through projection onto the modified velocity
 			if (!Sample.LinearVelocity.IsZero() && !Sample.Transform.GetLocation().IsZero())
 			{
diff --git a/Engine/Plugins/Experimental/Animation/MotionTrajectory/Source/MotionTrajectory/Public/MotionTrajectoryCharacterMovement.h b/Engine/Plugins/Experimental/Animation/MotionTrajectory/Source/MotionTrajectory/Public/MotionTrajectoryCharacterMovement.h
--- a/Engine/Plugins/Experimental/Animation/MotionTrajectory/Source/MotionTrajectory/Public/MotionTrajectoryCharacterMovement.h
+++ b/Engine/Plugins/Experimental/Animation/MotionTrajectory/Source/MotionTrajectory/Public/MotionTrajectoryCharacterMovement.h
@@ -6,6 +6,8 @@
 
 #include "MotionTrajectoryCharacterMovement.generated.h"
 
+class UCharacterMovementComponent;
+
 // Example motion trajectory component implementation for encapsulating: Character Movement ground locomotion
 UCLASS(meta=(BlueprintSpawnableComponent), Category="Motion Trajectory")
 class MOTIONTRAJECTORY_API UCharacterMovementTrajectoryComponent : public UMotionTrajectoryComponent
@@ -42,4 +44,50 @@ public:
 	virtual FTrajectorySampleRange GetTrajectory() const override;
 	virtual FTrajectorySampleRange GetTrajectoryWithSettings(const FMotionTrajectorySettings& Settings, bool bIncludeHistory) const override;
 	// End UMotionTrajectoryComponent Interface
+
+
+protected:
+	// Predicts future trajectory and writes prediction to ReturnValue
+	virtual void PredictTrajectory(
+		int32 SampleRate,
+		int32 MaxSamples,
+		const FMotionTrajectorySettings& Settings,
+		const FTrajectorySample& PresentTrajectory,
+		const FRotator& DesiredControlRotationVelocity,
+		FTrajectorySampleRange& OutTrajectoryRange) const;
+
+	// Updates InOutSample predicting IntegrationDelta forward
+	virtual void StepPrediction(
+		float IntegrationDelta,
+		const FRotator& ControlRotationVelocity,
+		FRotator& InOutControlRotationTotalDelta,
+		FTrajectorySample& InOutSample) const;
+
+	// The methods below allow tweaking movement model values per prediction sample
+
+	// Returns the friction to be used when updating Sample
+	virtual float GetFriction(
+		const UCharacterMovementComponent* MoveComponent, 
+		const FTrajectorySample& Sample, 
+		float DeltaSeconds) const;
+
+	// Returns maximum deceleration of Sample when the character is braking
+	virtual float GetMaxBrakingDeceleration(
+		const UCharacterMovementComponent* MoveComponent, 
+		const FTrajectorySample& Sample, 
+		float DeltaSeconds) const;
+
+	// Returns input acceleration in world space when predicting how Sample will change DeltaSeconds in the future
+	virtual FVector GetAccelerationWS(
+		const UCharacterMovementComponent* MoveComponent, 
+		const FTrajectorySample& Sample, 
+		float DeltaSeconds) const;
+
+	// This function returns the desired rotation for Sample without taking control rotation velocity into consideration.
+	// Control rotation velocity integration will be added to the result of this function when appropriate.
+	virtual FRotator GetBaseRotationWS(
+		const UCharacterMovementComponent* MoveComponent,
+		const FTrajectorySample& Sample,
+		const FRotator& SampleBaseRotationWS,
+		float DeltaSeconds) const;
 };
\ No newline at end of file
diff --git a/Engine/Source/Runtime/Engine/Private/Animation/MotionTrajectoryTypes.cpp b/Engine/Source/Runtime/Engine/Private/Animation/MotionTrajectoryTypes.cpp
--- a/Engine/Source/Runtime/Engine/Private/Animation/MotionTrajectoryTypes.cpp
+++ b/Engine/Source/Runtime/Engine/Private/Animation/MotionTrajectoryTypes.cpp
@@ -13,7 +13,6 @@ static constexpr int32 DebugTrajectorySampleTime = 2;
 static constexpr int32 DebugTrajectorySampleDistance = 3;
 static constexpr int32 DebugTrajectorySamplePosition = 4;
 static constexpr int32 DebugTrajectorySampleVelocity = 5;
-static constexpr int32 DebugTrajectorySampleAccel = 6;
 static const FVector DebugSampleTypeOffset(0.f, 0.f, 50.f);
 static const FVector DebugSampleOffset(0.f, 0.f, 10.f);
 
@@ -58,7 +57,6 @@ bool FTrajectorySample::IsZeroSample() const
 {
 	// AccumulatedTime is specifically omitted here to allow for the zero sample semantic across an entire trajectory range
 	return LinearVelocity.IsNearlyZero()
-		&& LinearAcceleration.IsNearlyZero()
 		&& Transform.GetTranslation().IsNearlyZero()
 		&& FMath::IsNearlyZero(AccumulatedDistance)
 		&& Transform.GetRotation().IsIdentity();
@@ -70,7 +68,6 @@ FTrajectorySample FTrajectorySample::Lerp(const FTrajectorySample& Sample, float
 	Interp.AccumulatedSeconds = FMath::Lerp(AccumulatedSeconds, Sample.AccumulatedSeconds, Alpha);
 	Interp.AccumulatedDistance = FMath::Lerp(AccumulatedDistance, Sample.AccumulatedDistance, Alpha);
 	Interp.LinearVelocity = FMath::Lerp(LinearVelocity, Sample.LinearVelocity, Alpha);
-	Interp.LinearAcceleration = FMath::Lerp(LinearAcceleration, Sample.LinearAcceleration, Alpha);
 
 	Interp.Transform.Blend(Transform, Sample.Transform, Alpha);
 	
@@ -86,7 +83,6 @@ FTrajectorySample FTrajectorySample::SmoothInterp(const FTrajectorySample& PrevS
 	Interp.AccumulatedDistance = CubicCRSplineInterpSafe(PrevSample.AccumulatedDistance, AccumulatedDistance, Sample.AccumulatedDistance, NextSample.AccumulatedDistance, Alpha);
 	Interp.AccumulatedSeconds = CubicCRSplineInterpSafe(PrevSample.AccumulatedSeconds, AccumulatedSeconds, Sample.AccumulatedSeconds, NextSample.AccumulatedSeconds, Alpha);
 	Interp.LinearVelocity = CubicCRSplineInterpSafe(PrevSample.LinearVelocity, LinearVelocity, Sample.LinearVelocity, NextSample.LinearVelocity, Alpha);
-	Interp.LinearAcceleration = CubicCRSplineInterpSafe(PrevSample.LinearAcceleration, LinearAcceleration, Sample.LinearAcceleration, NextSample.LinearAcceleration, Alpha);
 
 	Interp.Transform.SetLocation(CubicCRSplineInterpSafe(
 		PrevSample.Transform.GetLocation(),
@@ -130,17 +126,13 @@ void FTrajectorySample::PrependOffset(const FTransform DeltaTransform, float Del
 	}
 
 	Transform *= DeltaTransform;
-
 	LinearVelocity = DeltaTransform.TransformVectorNoScale(LinearVelocity);
-	LinearAcceleration = DeltaTransform.TransformVectorNoScale(LinearAcceleration);
 }
 
 void FTrajectorySample::TransformReferenceFrame(const FTransform DeltaTransform)
 {
 	Transform = DeltaTransform.Inverse() * Transform * DeltaTransform;
-
 	LinearVelocity = DeltaTransform.TransformVectorNoScale(LinearVelocity);
-	LinearAcceleration = DeltaTransform.TransformVectorNoScale(LinearAcceleration);
 }
 
 bool FTrajectorySampleRange::HasSamples() const
@@ -228,7 +220,7 @@ void FTrajectorySampleRange::DebugDrawTrajectory(bool bEnable
 						World, SampleTransformWS.GetLocation(), SampleTransformWS.Rotator(),
 						TransformScale, false, 0.f, 0, TransformThickness);
 
-					if (VelArrowThickness == 0.0f)
+					if (VelArrowScale == 0.0f)
 					{
 						DrawDebugSphere(
 							World, SamplePositionWS,										
@@ -261,10 +253,6 @@ void FTrajectorySampleRange::DebugDrawTrajectory(bool bEnable
 					DebugString = "Sample Velocity:";
 					DebugSampleString = DebugSampleString.Format(TEXT("{0}"), { Samples[Idx].LinearVelocity.ToCompactString() });
 					break;
-				case DebugTrajectorySampleAccel: // Sample Acceleration
-					DebugString = "Sample Acceleration:";
-					DebugSampleString = DebugSampleString.Format(TEXT("{0}"), { Samples[Idx].LinearAcceleration.ToCompactString() });
-					break;
 				default:
 					break;
 				}
diff --git a/Engine/Source/Runtime/Engine/Public/Animation/MotionTrajectoryTypes.h b/Engine/Source/Runtime/Engine/Public/Animation/MotionTrajectoryTypes.h
--- a/Engine/Source/Runtime/Engine/Public/Animation/MotionTrajectoryTypes.h
+++ b/Engine/Source/Runtime/Engine/Public/Animation/MotionTrajectoryTypes.h
@@ -40,10 +40,6 @@ struct ENGINE_API FTrajectorySample
 	UPROPERTY(EditAnywhere, BlueprintReadWrite, Category="Motion Trajectory")
 	FVector LinearVelocity = FVector::ZeroVector;
 
-	// Linear acceleration relative to the sampled in-motion object
-	UPROPERTY(EditAnywhere, BlueprintReadWrite, Category="Motion Trajectory")
-	FVector LinearAcceleration = FVector::ZeroVector;
-
 	// Linear interpolation of all parameters of two trajectory samples
 	FTrajectorySample Lerp(const FTrajectorySample& Sample, float Alpha) const;
 
