- Fixed trajectory generation problems:
   - Removed angular velocity glitch
   - Updated logic so constant turning produces a correct circular trajectory
- Made movement component variables BlueprintReadWrite so we can easily test different configurations

diff --git a/Engine/Plugins/Experimental/Animation/MotionTrajectory/Source/MotionTrajectory/Private/MotionTrajectoryCharacterMovement.cpp b/Engine/Plugins/Experimental/Animation/MotionTrajectory/Source/MotionTrajectory/Private/MotionTrajectoryCharacterMovement.cpp
--- a/Engine/Plugins/Experimental/Animation/MotionTrajectory/Source/MotionTrajectory/Private/MotionTrajectoryCharacterMovement.cpp
+++ b/Engine/Plugins/Experimental/Animation/MotionTrajectory/Source/MotionTrajectory/Private/MotionTrajectoryCharacterMovement.cpp
@@ -96,7 +96,10 @@ void UCharacterMovementTrajectoryComponent::TickTrajectory(float DeltaTime)
 		if (const APawn* Pawn = TryGetOwnerPawn())
 		{
 			const FRotator CurrentDesiredRotation = Pawn->Controller->GetDesiredRotation();
-			DesiredControlRotationVelocity = (CurrentDesiredRotation - LastDesiredControlRotation) * (1.0f / DeltaTime);
+
+			DesiredControlRotationVelocity = 
+				(CurrentDesiredRotation - LastDesiredControlRotation).GetNormalized() * (1.0f / DeltaTime);
+
 			LastDesiredControlRotation = CurrentDesiredRotation;
 		}
 	}
@@ -278,15 +281,20 @@ void UCharacterMovementTrajectoryComponent::StepPrediction(
 	FRotator& InOutControlRotationTotalDelta,
 	FTrajectorySample& InOutSample) const
 {
-	const FTransform InitialSampleTransform = InOutSample.Transform;
-
-	const UCharacterMovementComponent* MovementComponent = 
+	const UCharacterMovementComponent* MovementComponent =
 		Cast<UCharacterMovementComponent>(TryGetOwnerPawn()->GetMovementComponent());
 	if (!IsValid(MovementComponent))
 	{
 		return;
 	}
 
+	const FTransform InitialSampleTransform = InOutSample.Transform;
+	const FTransform PrevTransformWS = InitialSampleTransform * MovementComponent->GetOwner()->GetActorTransform();
+	
+	FRotator ControlRotationDelta = ControlRotationVelocity * IntegrationDelta;
+	ControlRotationDelta = AdjustRotationVerticality(ControlRotationDelta, MovementComponent);
+	InOutControlRotationTotalDelta += ControlRotationDelta;
+
 	if (MovementComponent->GetCurrentAcceleration().IsZero())
 	{
 		float ActualBrakingFriction = 
@@ -342,14 +350,16 @@ void UCharacterMovementTrajectoryComponent::StepPrediction(
 		const FVector LinearAccelerationWS = GetAccelerationWS(MovementComponent, InOutSample, IntegrationDelta);
 		const FVector LinearAccelerationAS = 
 			TryGetOwnerPawn()->GetActorTransform().Inverse().TransformVectorNoScale(LinearAccelerationWS);
+
 		const FVector RotatedLinearAccelerationAS = InOutControlRotationTotalDelta.RotateVector(LinearAccelerationAS);
 
 		const FVector AccelDir = RotatedLinearAccelerationAS.GetSafeNormal();
 		const float VelSize = InOutSample.LinearVelocity.Size();
+
 		InOutSample.LinearVelocity = 
 			InOutSample.LinearVelocity - 
 			(InOutSample.LinearVelocity - AccelDir * VelSize) * 
-			FMath::Min(IntegrationDelta * MovementComponent->GroundFriction, 1.f);
+			FMath::Min(IntegrationDelta * GetFriction(MovementComponent, InOutSample, IntegrationDelta), 1.f);
 
 		const float MaxInputSpeed = FMath::Max(
 			MovementComponent->GetMaxSpeed() * MovementComponent->GetAnalogInputModifier(), 
@@ -358,25 +368,22 @@ void UCharacterMovementTrajectoryComponent::StepPrediction(
 		InOutSample.LinearVelocity = InOutSample.LinearVelocity.GetClampedToMaxSize(MaxInputSpeed);
 	}
 
-	const FTransform PrevTransformWS = InitialSampleTransform * MovementComponent->GetOwner()->GetActorTransform();
-
 	{
 		const FRotator PrevRotator = PrevTransformWS.GetRotation().Rotator() - InOutControlRotationTotalDelta;
 		PrevRotator.DiagnosticCheckNaN(TEXT("StepPrediction: PrevRotator"));
 
-		FRotator BaseRotation = GetBaseRotationWS(MovementComponent, InOutSample, PrevRotator, IntegrationDelta);
+		const FRotator TotalRotation = StepRotationWS(
+			MovementComponent, 
+			InOutSample, 
+			PrevRotator,
+			InOutControlRotationTotalDelta,
+			IntegrationDelta);
 
-		FRotator ControlRotationDelta = ControlRotationVelocity * IntegrationDelta;
-		ControlRotationDelta = AdjustRotationVerticality(ControlRotationDelta, MovementComponent);
-		InOutControlRotationTotalDelta += ControlRotationDelta;
-
-		const FRotator TotalRotation = BaseRotation + InOutControlRotationTotalDelta;
 		const FQuat TotalRotationQuat = TotalRotation.Quaternion();
 		const FQuat TotalRotationQuatQuatAS = 
 			MovementComponent->GetOwner()->GetActorQuat().Inverse() * TotalRotationQuat;
-		
+
 		InOutSample.Transform.SetRotation(TotalRotationQuatQuatAS);
-		InOutSample.LinearVelocity = ControlRotationDelta.RotateVector(InOutSample.LinearVelocity);
 	
 		const FVector Translation = InOutSample.LinearVelocity * IntegrationDelta;
 		InOutSample.Transform.AddToTranslation(Translation);
@@ -407,29 +414,33 @@ FVector UCharacterMovementTrajectoryComponent::GetAccelerationWS(
 	return MoveComponent->GetCurrentAcceleration();
 }
 
-FRotator UCharacterMovementTrajectoryComponent::GetBaseRotationWS(
-	const UCharacterMovementComponent* MoveComponent,
+FRotator UCharacterMovementTrajectoryComponent::StepRotationWS(
+	const UCharacterMovementComponent* MovementComponent,
 	const FTrajectorySample& Sample,
-	const FRotator& SampleBaseRotationWS,
-	float DeltaSeconds) const
+	const FRotator& PrevRotator,
+	const FRotator& ControlRotationTotalDelta,
+	float IntegrationDelta) const
 {
 	FRotator BaseRotation = FRotator::ZeroRotator;
-	if (MoveComponent->bOrientRotationToMovement)
+	if (MovementComponent->bOrientRotationToMovement)
 	{
-		FRotator DeltaRot = MoveComponent->GetDeltaRotation(DeltaSeconds);
+		FRotator DeltaRot = MovementComponent->GetDeltaRotation(IntegrationDelta);
 		DeltaRot.DiagnosticCheckNaN(TEXT("StepPrediction: MovementComponent->GetDeltaRotation"));
 
 		BaseRotation =
-			MoveComponent->ComputeOrientToMovementRotation(SampleBaseRotationWS, DeltaSeconds, DeltaRot);
-		BaseRotation = AdjustRotationVerticality(BaseRotation, MoveComponent);
-		BaseRotation = LimitRotationRate(BaseRotation, SampleBaseRotationWS, DeltaRot, MoveComponent);
+			MovementComponent->ComputeOrientToMovementRotation(PrevRotator, IntegrationDelta, DeltaRot);
+		BaseRotation = AdjustRotationVerticality(BaseRotation, MovementComponent);
+		BaseRotation = LimitRotationRate(BaseRotation, PrevRotator, DeltaRot, MovementComponent);
 	}
 	else
 	{
-		BaseRotation = MoveComponent->GetOwner()->GetActorTransform().Rotator();
+		BaseRotation = MovementComponent->GetOwner()->GetActorTransform().Rotator();
 	}
 
-	return BaseRotation;
+	const FRotator TotalRotation = BaseRotation + ControlRotationTotalDelta;
+	return TotalRotation;
 }
 
+
+
 // ------------ END Derived from FCharacterMovementComponentAsyncInput ------------ //
diff --git a/Engine/Plugins/Experimental/Animation/MotionTrajectory/Source/MotionTrajectory/Public/MotionTrajectoryCharacterMovement.h b/Engine/Plugins/Experimental/Animation/MotionTrajectory/Source/MotionTrajectory/Public/MotionTrajectoryCharacterMovement.h
--- a/Engine/Plugins/Experimental/Animation/MotionTrajectory/Source/MotionTrajectory/Public/MotionTrajectoryCharacterMovement.h
+++ b/Engine/Plugins/Experimental/Animation/MotionTrajectory/Source/MotionTrajectory/Public/MotionTrajectoryCharacterMovement.h
@@ -83,11 +83,12 @@ protected:
 		const FTrajectorySample& Sample, 
 		float DeltaSeconds) const;
 
-	// This function returns the desired rotation for Sample without taking control rotation velocity into consideration.
-	// Control rotation velocity integration will be added to the result of this function when appropriate.
-	virtual FRotator GetBaseRotationWS(
-		const UCharacterMovementComponent* MoveComponent,
+	// This function outputs a new rotator obtained by integrating Sample forward by IntegrationDelta based on the
+	// dynamics specified in MovementComponent
+	virtual FRotator StepRotationWS(
+		const UCharacterMovementComponent* MovementComponent,
 		const FTrajectorySample& Sample,
-		const FRotator& SampleBaseRotationWS,
-		float DeltaSeconds) const;
+		const FRotator& PrevRotator,
+		const FRotator& ControlRotationTotalDelta,
+		float IntegrationDelta) const;
 };
\ No newline at end of file
