[UE-141254] Fix for crash when using a pawn without a controller but with CharacterMovementTrajectoryComponent attached.

diff --git a/Engine/Plugins/Experimental/Animation/MotionTrajectory/Source/MotionTrajectory/Private/MotionTrajectoryCharacterMovement.cpp b/Engine/Plugins/Experimental/Animation/MotionTrajectory/Source/MotionTrajectory/Private/MotionTrajectoryCharacterMovement.cpp
--- a/Engine/Plugins/Experimental/Animation/MotionTrajectory/Source/MotionTrajectory/Private/MotionTrajectoryCharacterMovement.cpp
+++ b/Engine/Plugins/Experimental/Animation/MotionTrajectory/Source/MotionTrajectory/Private/MotionTrajectoryCharacterMovement.cpp
@@ -20,13 +20,23 @@ FTrajectorySample UCharacterMovementTrajectoryComponent::CalcWorldSpacePresentTr
 	}
 
 	const UCharacterMovementComponent* MovementComponent = Cast<UCharacterMovementComponent>(Pawn->GetMovementComponent());
-	if (ensure(MovementComponent) && MovementComponent->MovementMode == EMovementMode::MOVE_Walking)
+	if (MovementComponent)
 	{
-		FTransform ComponentWorldTransform = Pawn->GetActorTransform();
-		FVector MovementComponentPosition = MovementComponent->GetLastUpdateLocation();
+		if (MovementComponent->MovementMode == EMovementMode::MOVE_Walking)
+		{
+			FTransform ComponentWorldTransform = Pawn->GetActorTransform();
+			FVector MovementComponentPosition = MovementComponent->GetLastUpdateLocation();
 
-		ReturnValue.Transform = ComponentWorldTransform;
-		ReturnValue.LinearVelocity = MovementComponent->Velocity;
+			ReturnValue.Transform = ComponentWorldTransform;
+			ReturnValue.LinearVelocity = MovementComponent->Velocity;
+		}
+	}
+	else
+	{
+		UE_LOG(
+			LogMotionTrajectory,
+			Error,
+			TEXT("UCharacterMovementTrajectoryComponent expects the owner to have a CharacterMovementComponent"));
 	}
 
 	return ReturnValue;
@@ -62,11 +72,21 @@ void UCharacterMovementTrajectoryComponent::InitializeComponent()
 void UCharacterMovementTrajectoryComponent::UninitializeComponent()
 {
 	UCharacterMovementComponent* CharacterMovementComponent = GetOwner()->FindComponentByClass<UCharacterMovementComponent>();
-	if (ACharacter* Character = CharacterMovementComponent->GetCharacterOwner())
+	if (CharacterMovementComponent)
 	{
-		Character->OnCharacterMovementUpdated.RemoveDynamic(
-			this,
-			&UCharacterMovementTrajectoryComponent::OnMovementUpdated);
+		if (ACharacter* Character = CharacterMovementComponent->GetCharacterOwner())
+		{
+			Character->OnCharacterMovementUpdated.RemoveDynamic(
+				this,
+				&UCharacterMovementTrajectoryComponent::OnMovementUpdated);
+		}
+	}
+	else
+	{
+		UE_LOG(
+			LogMotionTrajectory, 
+			Error, 
+			TEXT("UCharacterMovementTrajectoryComponent expects the owner to have a CharacterMovementComponent"));
 	}
 
 	Super::UninitializeComponent();
@@ -76,7 +96,8 @@ void UCharacterMovementTrajectoryComponent::BeginPlay()
 {
 	Super::BeginPlay();
 
-	if (const APawn* Pawn = TryGetOwnerPawn())
+	const APawn* Pawn = TryGetOwnerPawn();
+	if (Pawn && Pawn->Controller)
 	{
 		LastDesiredControlRotation = Pawn->Controller->GetDesiredRotation();
 	}
@@ -93,8 +114,10 @@ void UCharacterMovementTrajectoryComponent::TickTrajectory(float DeltaTime)
 	Super::TickTrajectory(DeltaTime);
 	if (DeltaTime > SMALL_NUMBER)
 	{
-		if (const APawn* Pawn = TryGetOwnerPawn())
+		const APawn* Pawn = TryGetOwnerPawn();
+		if (Pawn && Pawn->Controller)
 		{
+
 			const FRotator CurrentDesiredRotation = Pawn->Controller->GetDesiredRotation();
 
 			DesiredControlRotationVelocity = 
