Add stop filtering to motion matching demo in QAGame

diff --git a/Engine/Plugins/Experimental/Animation/MotionTrajectory/Source/MotionTrajectory/Private/MotionTrajectoryLibrary.cpp b/Engine/Plugins/Experimental/Animation/MotionTrajectory/Source/MotionTrajectory/Private/MotionTrajectoryLibrary.cpp
--- a/Engine/Plugins/Experimental/Animation/MotionTrajectory/Source/MotionTrajectory/Private/MotionTrajectoryLibrary.cpp
+++ b/Engine/Plugins/Experimental/Animation/MotionTrajectory/Source/MotionTrajectory/Private/MotionTrajectoryLibrary.cpp
@@ -246,4 +246,102 @@ void UMotionTrajectoryBlueprintLibrary::DebugDrawTrajectory(const AActor* Actor
 			, ArrowThickness
 		);
 	}
-}
\ No newline at end of file
+}
+
+bool UMotionTrajectoryBlueprintLibrary::IsStoppingTrajectory(
+	const FTrajectorySampleRange& Trajectory, 
+	float MoveMinSpeed,
+	float IdleMaxSpeed)
+{
+	if (Trajectory.Samples.Num() > 0)
+	{
+		const FVector LastLinearVelocity = Trajectory.Samples.Last().LinearVelocity;
+		const float SquaredLastLinearSpeed = LastLinearVelocity.SquaredLength();
+		
+		int StartIdx = 0;
+		const FTrajectorySample& PresentSample = FTrajectorySampleRange::IterSampleTrajectory(
+			Trajectory.Samples, 
+			ETrajectorySampleDomain::Time, 
+			0.0f, 
+			StartIdx);
+		const FVector PresenLinearVelocity = PresentSample.LinearVelocity;
+		const float SquaredPresentLinearSpeed = PresenLinearVelocity.SquaredLength();
+
+		const bool bIsStopping =
+			(SquaredPresentLinearSpeed >= (MoveMinSpeed * MoveMinSpeed)) &&
+			(SquaredLastLinearSpeed <= (IdleMaxSpeed * IdleMaxSpeed));
+
+		return bIsStopping;
+	}
+
+	return false;
+}
+
+bool UMotionTrajectoryBlueprintLibrary::IsStartingTrajectory(
+	const FTrajectorySampleRange& Trajectory, 
+	float MoveMinSpeed,
+	float IdleMaxSpeed)
+{
+	if (Trajectory.Samples.Num() > 0)
+	{
+		const FVector FirstLinearVelocity = Trajectory.Samples[0].LinearVelocity;
+		const float SquaredFirstLinearSpeed = FirstLinearVelocity.SquaredLength();
+
+		int StartIdx = 0;
+		const FTrajectorySample& PresentSample = FTrajectorySampleRange::IterSampleTrajectory(
+			Trajectory.Samples,
+			ETrajectorySampleDomain::Time,
+			0.0f,
+			StartIdx);
+		const FVector PresentLinearVelocity = PresentSample.LinearVelocity;
+		const float SquaredPresentLinearSpeed = PresentLinearVelocity.SquaredLength();
+
+		const bool IsStarting =
+			(SquaredPresentLinearSpeed >= (MoveMinSpeed * MoveMinSpeed)) &&
+			(SquaredFirstLinearSpeed <= (IdleMaxSpeed * IdleMaxSpeed));
+
+		return IsStarting;
+	}
+
+	return false;
+}
+
+
+bool UMotionTrajectoryBlueprintLibrary::IsConstantSpeedTrajectory(
+	const FTrajectorySampleRange& Trajectory,
+	float Speed,
+	float Tolerance)
+{
+	if (Trajectory.Samples.Num() > 0)
+	{
+		const float MinSpeed = FMath::Max(Speed - Tolerance, 0.0f);
+		const float SquaredMinSpeed = MinSpeed * MinSpeed;
+		const float MaxSpeed = FMath::Max(Speed + Tolerance, 0.0f);
+		const float SquaredMaxSpeed = MaxSpeed * MaxSpeed;
+
+		auto IsWithinLimit = [=](float SquaredSpeed)
+		{
+			return SquaredSpeed >= SquaredMinSpeed && SquaredSpeed <= SquaredMaxSpeed;
+		};
+
+		const FVector LastLinearVelocity = Trajectory.Samples.Last().LinearVelocity;
+		const float SquaredLastLinearSpeed = LastLinearVelocity.SquaredLength();
+		const bool LastSpeedWithinLimit = IsWithinLimit(SquaredLastLinearSpeed);
+
+		int StartIdx = 0;
+		const FTrajectorySample& PresentSample = FTrajectorySampleRange::IterSampleTrajectory(
+			Trajectory.Samples,
+			ETrajectorySampleDomain::Time,
+			0.0f,
+			StartIdx);
+		const FVector PresentLinearVelocity = PresentSample.LinearVelocity;
+		const float SquaredPresentLinearSpeed = PresentLinearVelocity.SquaredLength();
+		const bool PresentSpeedWithinLimit = IsWithinLimit(SquaredLastLinearSpeed);
+
+		const bool bIsAtSpeed = LastSpeedWithinLimit && PresentSpeedWithinLimit;
+
+		return bIsAtSpeed;
+	}
+
+	return false;
+}
diff --git a/Engine/Plugins/Experimental/Animation/MotionTrajectory/Source/MotionTrajectory/Public/MotionTrajectoryLibrary.h b/Engine/Plugins/Experimental/Animation/MotionTrajectory/Source/MotionTrajectory/Public/MotionTrajectoryLibrary.h
--- a/Engine/Plugins/Experimental/Animation/MotionTrajectory/Source/MotionTrajectory/Public/MotionTrajectoryLibrary.h
+++ b/Engine/Plugins/Experimental/Animation/MotionTrajectory/Source/MotionTrajectory/Public/MotionTrajectoryLibrary.h
@@ -103,4 +103,28 @@ public:
 		, float ArrowSize = 40.f
 		, float ArrowThickness = 2.f
 	);
+
+	// Returns true if the furthest future sample, at the end of the trajectory, has velocity magnitude below 
+	// IdleMaxSpeed and present speed above MoveMinSpeed.
+	UFUNCTION(BlueprintPure, Category = "Motion Trajectory", meta = (BlueprintThreadSafe))
+	static bool IsStoppingTrajectory(
+		const FTrajectorySampleRange& Trajectory, 
+		float MoveMinSpeed = 1.0f,
+		float IdleMaxSpeed = 0.0f);
+
+	// Returns true if the furthest past sample, at the beginning of the trajectory, has velocity magnitude below 
+	// IdleMaxSpeed and present speed above MoveMinSpeed.
+	UFUNCTION(BlueprintPure, Category = "Motion Trajectory", meta = (BlueprintThreadSafe))
+	static bool IsStartingTrajectory(
+		const FTrajectorySampleRange& Trajectory, 
+		float MoveMinSpeed = 1.0f,
+		float IdleMaxSpeed = 0.0f);
+
+	// Returns true if the furthest future sample, at the end of the trajectory, and the present sample, have a velocity
+	// magnitude that is within Tolerance of Speed.
+	UFUNCTION(BlueprintPure, Category = "Motion Trajectory", meta = (BlueprintThreadSafe))
+	static bool IsConstantSpeedTrajectory(
+			const FTrajectorySampleRange& Trajectory,
+			float Speed = 0.0f,
+			float Tolerance = 0.001f);
 };
\ No newline at end of file
