Implemented motion trajectory turn detection so that can be used to select different motion matching filters
- Created new movement component so the trajectory prediction can incorporate changes previously implemented in blueprint. Includes:
   - MMCharacter
   - MMCharacterTrajectoryComponent
   - MMCharacterMovementComponent
- Implemented turn detection function in MotionTrajectoryLibrary

diff --git a/Engine/Plugins/Experimental/Animation/MotionTrajectory/Source/MotionTrajectory/Private/MotionTrajectoryLibrary.cpp b/Engine/Plugins/Experimental/Animation/MotionTrajectory/Source/MotionTrajectory/Private/MotionTrajectoryLibrary.cpp
--- a/Engine/Plugins/Experimental/Animation/MotionTrajectory/Source/MotionTrajectory/Private/MotionTrajectoryLibrary.cpp
+++ b/Engine/Plugins/Experimental/Animation/MotionTrajectory/Source/MotionTrajectory/Private/MotionTrajectoryLibrary.cpp
@@ -278,7 +278,7 @@ bool UMotionTrajectoryBlueprintLibrary::IsStoppingTrajectory(
 }
 
 bool UMotionTrajectoryBlueprintLibrary::IsStartingTrajectory(
-	const FTrajectorySampleRange& Trajectory, 
+	const FTrajectorySampleRange& Trajectory,
 	float MoveMinSpeed,
 	float IdleMaxSpeed)
 {
@@ -345,3 +345,209 @@ bool UMotionTrajectoryBlueprintLibrary::IsConstantSpeedTrajectory(
 
 	return false;
 }
+
+namespace UE::MotionTrajectory
+{
+
+	// Populate arrays with turn information. If these end up being used multiple times per frame, we should consider
+	// adding this information to FTrajectorySample (maybe deferred until needed)
+
+	template <typename T>
+	void CalcTurnData(
+		const FTrajectorySampleRange& Trajectory,
+		const FVector& TurnAxis,
+		T& OutAccumulatedTurnAxisRotations,
+		T& OutImmediateTurnAxisRotations)
+	{
+		OutAccumulatedTurnAxisRotations.SetNum(Trajectory.Samples.Num());
+		OutAccumulatedTurnAxisRotations[0] = 0.0f;
+
+		OutImmediateTurnAxisRotations.SetNum(Trajectory.Samples.Num());
+		OutImmediateTurnAxisRotations[0] = 0.0f;
+
+		const int32 NumSamples = Trajectory.Samples.Num();
+		for (int32 SampleIdx = 1; SampleIdx < NumSamples; ++SampleIdx)
+		{
+			const FTrajectorySample& PrevSample = Trajectory.Samples[SampleIdx - 1];
+			const FTrajectorySample& CurrSample = Trajectory.Samples[SampleIdx];
+
+			FQuat VelocityDelta = FQuat::FindBetweenVectors(PrevSample.LinearVelocity, CurrSample.LinearVelocity);
+			FVector DeltaAxis;
+			float DeltaAngle = VelocityDelta.GetTwistAngle(TurnAxis);
+
+			OutImmediateTurnAxisRotations[SampleIdx] = DeltaAngle;
+			OutAccumulatedTurnAxisRotations[SampleIdx] = OutAccumulatedTurnAxisRotations[SampleIdx - 1] + DeltaAngle;
+		}
+	}
+
+	// calculate rotation speed based on accumulated rotation info
+	float CalcRotationSpeed(
+		const FTrajectorySampleRange& Trajectory,
+		const TArrayView<float>& AccumulatedRotations,
+		int32 StartIdx,
+		int32 EndIdx)
+	{
+		const FTrajectorySample& SampleA = Trajectory.Samples[StartIdx];
+		const FTrajectorySample& SampleB = Trajectory.Samples[EndIdx];
+		const float TotalTime = SampleB.AccumulatedSeconds - SampleA.AccumulatedSeconds;
+		const float TotalAngleDelta = AccumulatedRotations[EndIdx] - AccumulatedRotations[StartIdx];
+		const float RotationSpeed = TotalAngleDelta / TotalTime;
+		return RotationSpeed;
+	}
+
+	bool FindTurnBeyondExtrapolation(
+		const FTrajectorySampleRange& Trajectory,
+		const TArrayView<float>& AccumulatedTurnAxisRotations,
+		const TArrayView<float>& ImmediateTurnAxisRotations,
+		float RotSpeedToExtrapolate,
+		float MinSharpTurnAngleRadians)
+	{
+		check(Trajectory.Samples.Num() > 1);
+
+		const float FirstToLastDeltaSeconds =
+			Trajectory.Samples.Last().AccumulatedSeconds - Trajectory.Samples[0].AccumulatedSeconds;
+
+		const float TargetRotSpeed = MinSharpTurnAngleRadians / FirstToLastDeltaSeconds;
+
+		const int32 LastSampleIdx = Trajectory.Samples.Num() - 1;
+		const float FirstToLastRotSpeed = CalcRotationSpeed(Trajectory, AccumulatedTurnAxisRotations, 0, LastSampleIdx);
+
+		const float SignLastDirRotSpeed = FMath::Sign(RotSpeedToExtrapolate);
+		const float SignedFirstToLastRotSpeed = SignLastDirRotSpeed * FirstToLastRotSpeed;
+
+		if (SignedFirstToLastRotSpeed > SignLastDirRotSpeed * (RotSpeedToExtrapolate + TargetRotSpeed) ||
+			SignedFirstToLastRotSpeed < -TargetRotSpeed)
+		{
+			// under these conditions there's no need to search further, a turn must exist in the trajectory
+			return true;
+		}
+
+
+		// The case of a turn in the same direction of the circling speed RotSpeedToExtrapolate is already covered by the
+		// conditional above. The code below searches for a turn in the opposite direction only.
+
+		int32 FirstTurnSampleIdx = INDEX_NONE;
+		for (int32 SampleIdx = 0; SampleIdx < LastSampleIdx; ++SampleIdx)
+		{
+			if (SignLastDirRotSpeed * ImmediateTurnAxisRotations[SampleIdx] < 0.0f)
+			{
+				FirstTurnSampleIdx = SampleIdx;
+				break;
+			}
+		}
+
+		if (FirstTurnSampleIdx == INDEX_NONE)
+		{
+			// no turn to the expected side
+			return false;
+		}
+
+		int32 LastTurnSampleIdx = FirstTurnSampleIdx + 1;
+		for (int32 SampleIdx = FirstTurnSampleIdx + 2; SampleIdx <= LastSampleIdx; ++SampleIdx)
+		{
+			if (SignLastDirRotSpeed * ImmediateTurnAxisRotations[SampleIdx] >= 0.0f)
+			{
+				LastTurnSampleIdx = SampleIdx;
+				break;
+			}
+		}
+
+		const float AccumulatedTurnRotation =
+			AccumulatedTurnAxisRotations[LastTurnSampleIdx] - AccumulatedTurnAxisRotations[FirstTurnSampleIdx];
+
+		const bool bIsSharpTurn = SignLastDirRotSpeed * AccumulatedTurnRotation <= -MinSharpTurnAngleRadians;
+		return bIsSharpTurn;
+	}
+} // namespace UE::MotionTrajectory
+
+bool UMotionTrajectoryBlueprintLibrary::IsSharpVelocityDirChange(
+	const FTrajectorySampleRange& Trajectory,
+	float MinSharpTurnAngleDegrees,
+	float MaxAlignmentAngleDegrees,
+	float MinLinearSpeed,
+	FVector TurnAxis,
+	FVector ForwardAxis)
+{
+	if (Trajectory.Samples.Num() < 2)
+	{
+		// not enough samples to evaluate the trajectory
+		return false;
+	}
+	const float SquaredMinLinearSpeed = MinLinearSpeed * MinLinearSpeed;
+	const int32 LastSampleIdx = Trajectory.Samples.Num() - 1;
+
+	int32 PresentIdx = 0;
+	const FTrajectorySample& PresentSample = FTrajectorySampleRange::IterSampleTrajectory(
+		Trajectory.Samples,
+		ETrajectorySampleDomain::Time,
+		0.0f,
+		PresentIdx);
+
+	if (PresentIdx >= LastSampleIdx)
+	{
+		// we need at least one more sample in the future
+		return false;
+	}
+
+	const FTrajectorySample& LastSample = Trajectory.Samples[LastSampleIdx];
+	const float SquaredLastSampleSpeed = LastSample.LinearVelocity.SquaredLength();
+
+	const FTrajectorySample& FirstSample = Trajectory.Samples[0];
+	const float SquaredFirstSampleSpeed = FirstSample.LinearVelocity.SquaredLength();
+
+	if (SquaredLastSampleSpeed < SquaredMinLinearSpeed || SquaredFirstSampleSpeed < SquaredMinLinearSpeed)
+	{
+		// trajectory end points not fast enough
+		return false;
+	}
+
+	const FTrajectorySample& BeforeLastSample = Trajectory.Samples[LastSampleIdx - 1];
+
+	const float LastDeltaSeconds = LastSample.AccumulatedSeconds - BeforeLastSample.AccumulatedSeconds;
+	if (LastDeltaSeconds < SMALL_NUMBER)
+	{
+		// delta too small to evaluate
+		return false;
+	}
+
+	const FQuat LastDirectionRotation = 
+		FQuat::FindBetweenVectors(BeforeLastSample.LinearVelocity, LastSample.LinearVelocity);
+	const float LastDirRotAngle = LastDirectionRotation.GetTwistAngle(TurnAxis);
+	const float LastDirRotSpeed = LastDirRotAngle / LastDeltaSeconds;
+
+	TArray<float, TInlineAllocator<128>> AccumulatedTurnAxisRotations;
+	TArray<float, TInlineAllocator<128>> ImmediateTurnAxisRotations;
+	UE::MotionTrajectory::CalcTurnData(Trajectory, TurnAxis, AccumulatedTurnAxisRotations, ImmediateTurnAxisRotations);
+
+	const float PresentToLastRotSpeed = 
+		UE::MotionTrajectory::CalcRotationSpeed(Trajectory, AccumulatedTurnAxisRotations, PresentIdx, LastSampleIdx);
+	const float PresentSpeedDelta = PresentToLastRotSpeed - LastDirRotSpeed;
+	const float PresentAngleDelta = PresentSpeedDelta * LastSample.AccumulatedSeconds;
+
+	const FVector PresentFwd = PresentSample.Transform.GetRotation().RotateVector(ForwardAxis);
+	const FQuat PresentVelToFwdRotation = FQuat::FindBetweenVectors(PresentSample.LinearVelocity, PresentFwd);
+
+	const FVector LastFwd = LastSample.Transform.GetRotation().RotateVector(ForwardAxis);
+	const FQuat LastVelToFwdRotation = FQuat::FindBetweenVectors(LastSample.LinearVelocity, LastFwd);
+
+	const FQuat VelAlignmentDelta = PresentVelToFwdRotation.Inverse() * LastVelToFwdRotation;
+	const float VelAlignmentDeltaAngle = VelAlignmentDelta.GetTwistAngle(TurnAxis);
+
+	const float MaxAlignmentAngleRadians = FMath::DegreesToRadians(MaxAlignmentAngleDegrees);
+	if ((FMath::Abs(PresentAngleDelta) < MaxAlignmentAngleRadians) &&
+		(FMath::Abs(VelAlignmentDeltaAngle) < MaxAlignmentAngleRadians))
+	{
+		// the small difference indicates even if a sharp turn is in the trajectory past, it has already ended
+		return false;
+	}
+
+	const float MinSharpTurnAngleRadians = FMath::DegreesToRadians(MinSharpTurnAngleDegrees);
+	const bool bIsSharpTurn = UE::MotionTrajectory::FindTurnBeyondExtrapolation(
+		Trajectory, 
+		AccumulatedTurnAxisRotations, 
+		ImmediateTurnAxisRotations, 
+		LastDirRotSpeed, 
+		MinSharpTurnAngleRadians);
+
+	return bIsSharpTurn;
+}
diff --git a/Engine/Plugins/Experimental/Animation/MotionTrajectory/Source/MotionTrajectory/Public/MotionTrajectoryLibrary.h b/Engine/Plugins/Experimental/Animation/MotionTrajectory/Source/MotionTrajectory/Public/MotionTrajectoryLibrary.h
--- a/Engine/Plugins/Experimental/Animation/MotionTrajectory/Source/MotionTrajectory/Public/MotionTrajectoryLibrary.h
+++ b/Engine/Plugins/Experimental/Animation/MotionTrajectory/Source/MotionTrajectory/Public/MotionTrajectoryLibrary.h
@@ -112,19 +112,61 @@ public:
 		float MoveMinSpeed = 1.0f,
 		float IdleMaxSpeed = 0.0f);
 
-	// Returns true if the furthest past sample, at the beginning of the trajectory, has velocity magnitude below 
-	// IdleMaxSpeed and present speed above MoveMinSpeed.
+	/**
+	 * Returns true if the furthest past sample, at the beginning of the trajectory, has velocity magnitude below 
+	 * IdleMaxSpeed and present speed above MoveMinSpeed.
+	 * 
+	 * @param Trajectory	Trajectory being evaluated.
+	 * @param MoveMinSpeed	Minimum current speed required for this function to return true.
+	 * @param IdleMaxSpeed	Maximum speed at the furthest past sample for this function to return true.
+	 * 
+	 * @return True if a start is detected; false otherwise.
+	 */
 	UFUNCTION(BlueprintPure, Category = "Motion Trajectory", meta = (BlueprintThreadSafe))
 	static bool IsStartingTrajectory(
 		const FTrajectorySampleRange& Trajectory, 
 		float MoveMinSpeed = 1.0f,
 		float IdleMaxSpeed = 0.0f);
 
-	// Returns true if the furthest future sample, at the end of the trajectory, and the present sample, have a velocity
-	// magnitude that is within Tolerance of Speed.
+	/** 
+	 * Returns true if the furthest future sample, at the end of the trajectory, and the present sample, have a velocity
+	 * magnitude that is within Tolerance of Speed.
+	 * 
+	 * @param Trajectory	Trajectory being evaluated.
+	 * @param Speed			Speed that must be matched at present and future for this function to return true
+	 * @param Tolerance		How much the trajectory speed can deviate from Speed
+	 * 
+	 * @return True if the trajectory is at the desired speed.
+	 */
 	UFUNCTION(BlueprintPure, Category = "Motion Trajectory", meta = (BlueprintThreadSafe))
 	static bool IsConstantSpeedTrajectory(
 			const FTrajectorySampleRange& Trajectory,
 			float Speed = 0.0f,
 			float Tolerance = 0.001f);
+	/**
+	 * Returns true if the trajectory has a sharp velocity direction change. This function will compare the total 
+	 * turning of the trajectory against the extrapolation of the final angular velocity in the trajectory. This ensures
+	 * the function distinguishes sharp turns from smooth circling. Note this is not detecting facing changes,
+	 * this function will return true in a strafing turn in which the character is always facing the same direction.
+	 * 
+	 * @param Trajectory				Trajectory being evaluated.
+	 * @param MinSharpTurnAngleDegrees	How many degrees of turning must be in the trajectory that can't be explained by 
+	 *									extrapolating the angular velocity at the end of the trajectory.
+	 * @param MaxAlignmentAngleDegrees	If the current trajectory sample velocity and facing are well aligned with the 
+	 *									last point in the trajectory future, there's no turn.
+	 * @param MinLinearSpeed			Minimum linear speed at the end points of the trajectory required for this
+	 *									function to return true, so it's possible to ignore turning starts/stops
+	 * @param TurnAxis					Specifies the axis along which turns will be evaluated.
+	 * @param ForwardAxis				Specifies the trajectory forward angle, used to evaluate alignment.
+	 *									
+	 * @return True if a sharp turn is detected, false otherwise.
+	 */
+	UFUNCTION(BlueprintPure, Category = "Motion Trajectory", meta = (BlueprintThreadSafe))
+	static bool IsSharpVelocityDirChange(
+		const FTrajectorySampleRange& Trajectory,
+		float MinSharpTurnAngleDegrees = 45.0f,
+		float MaxAlignmentAngleDegrees = 5.0f,
+		float MinLinearSpeed = 1.0f,
+		FVector TurnAxis = FVector::UpVector,
+		FVector ForwardAxis = FVector::RightVector);
 };
\ No newline at end of file
