MotionTrajectory Library: New iteration of turn detection logic

diff --git a/Engine/Plugins/Experimental/Animation/MotionTrajectory/Source/MotionTrajectory/Private/MotionTrajectoryLibrary.cpp b/Engine/Plugins/Experimental/Animation/MotionTrajectory/Source/MotionTrajectory/Private/MotionTrajectoryLibrary.cpp
--- a/Engine/Plugins/Experimental/Animation/MotionTrajectory/Source/MotionTrajectory/Private/MotionTrajectoryLibrary.cpp
+++ b/Engine/Plugins/Experimental/Animation/MotionTrajectory/Source/MotionTrajectory/Private/MotionTrajectoryLibrary.cpp
@@ -348,121 +348,323 @@ bool UMotionTrajectoryBlueprintLibrary::IsConstantSpeedTrajectory(
 
 namespace UE::MotionTrajectory
 {
+	struct FTurnEvaluator
+	{
 
-	// Populate arrays with turn information. If these end up being used multiple times per frame, we should consider
-	// adding this information to FTrajectorySample (maybe deferred until needed)
+	public:
+		FTurnEvaluator(
+			const FTrajectorySampleRange& InTrajectory, 
+			ETrajectorySampleDomain InRotationConstraintDomain,
+			float InRotationConstraintLimit,
+			const FVector& InTurnAxis)
+			: Trajectory(InTrajectory)
+			, RotationConstraintDomain(InRotationConstraintDomain)
+			, RotationConstraintLimit(InRotationConstraintLimit)
+			, TurnAxis(InTurnAxis)
+		{
+			CalcRotSpeedToExtrapolate();
+			CalcTurnData();
+		}
 
-	template <typename T>
-	void CalcTurnData(
-		const FTrajectorySampleRange& Trajectory,
-		const FVector& TurnAxis,
-		T& OutAccumulatedTurnAxisRotations,
-		T& OutImmediateTurnAxisRotations)
-	{
-		OutAccumulatedTurnAxisRotations.SetNum(Trajectory.Samples.Num());
-		OutAccumulatedTurnAxisRotations[0] = 0.0f;
+		const FTrajectorySampleRange& Trajectory;
+		ETrajectorySampleDomain RotationConstraintDomain;
+		float RotationConstraintLimit;
+		FVector TurnAxis;
 
-		OutImmediateTurnAxisRotations.SetNum(Trajectory.Samples.Num());
-		OutImmediateTurnAxisRotations[0] = 0.0f;
+		float RotSpeedToExtrapolate;
+		TArray<float, TInlineAllocator<128>> AccumulatedTurnAxisRotations;
+		TArray<float, TInlineAllocator<128>> ImmediateTurnAxisRotations;
 
-		const int32 NumSamples = Trajectory.Samples.Num();
-		for (int32 SampleIdx = 1; SampleIdx < NumSamples; ++SampleIdx)
+		// Evaluates if there's a turn in the trajectory starting from SampleIdx. If the velocity and the alignment
+		// between the forward axis and the velocity match when comparing the samples at SampleIdx and the end of the
+		// trajectory, up to MaxAlignmentAngleDegrees, the function will return false. Otherwise, it returns false.
+		bool IsTurning(int32 SampleIdx, const FVector& ForwardAxis, float MaxAlignmentAngleDegrees)
 		{
-			const FTrajectorySample& PrevSample = Trajectory.Samples[SampleIdx - 1];
-			const FTrajectorySample& CurrSample = Trajectory.Samples[SampleIdx];
+			const FTrajectorySample CurrentSample = Trajectory.Samples[SampleIdx];
+
+			const int32 LastSampleIdx = Trajectory.Samples.Num() - 1;
+			const FTrajectorySample LastSample = Trajectory.Samples[LastSampleIdx];
+
+			const float CurrentToLastRotSpeed =
+				CalcRotationSpeed(SampleIdx, LastSampleIdx);
+			const float CurrentSpeedDelta = CurrentToLastRotSpeed - RotSpeedToExtrapolate;
+			const float CurrentAngleDelta = CurrentSpeedDelta * LastSample.AccumulatedSeconds;
+
+			const FVector CurrentFwd = CurrentSample.Transform.GetRotation().RotateVector(ForwardAxis);
+			const FQuat CurrentVelToFwdRotation = FQuat::FindBetweenVectors(CurrentSample.LinearVelocity, CurrentFwd);
+
+			const FVector LastFwd = LastSample.Transform.GetRotation().RotateVector(ForwardAxis);
+			const FQuat LastVelToFwdRotation = FQuat::FindBetweenVectors(LastSample.LinearVelocity, LastFwd);
 
-			FQuat VelocityDelta = FQuat::FindBetweenVectors(PrevSample.LinearVelocity, CurrSample.LinearVelocity);
-			FVector DeltaAxis;
-			float DeltaAngle = VelocityDelta.GetTwistAngle(TurnAxis);
+			const FQuat VelAlignmentDelta = CurrentVelToFwdRotation.Inverse() * LastVelToFwdRotation;
+			const float VelAlignmentDeltaAngle = VelAlignmentDelta.GetTwistAngle(TurnAxis);
 
-			OutImmediateTurnAxisRotations[SampleIdx] = DeltaAngle;
-			OutAccumulatedTurnAxisRotations[SampleIdx] = OutAccumulatedTurnAxisRotations[SampleIdx - 1] + DeltaAngle;
+			const float MaxAlignmentAngleRadians = FMath::DegreesToRadians(MaxAlignmentAngleDegrees);
+			bool bIsAligned =
+				(FMath::Abs(CurrentAngleDelta) < MaxAlignmentAngleRadians) &&
+				(FMath::Abs(VelAlignmentDeltaAngle) < MaxAlignmentAngleRadians);
+
+			return !bIsAligned;
 		}
-	}
 
-	// calculate rotation speed based on accumulated rotation info
-	float CalcRotationSpeed(
-		const FTrajectorySampleRange& Trajectory,
-		const TArrayView<float>& AccumulatedRotations,
-		int32 StartIdx,
-		int32 EndIdx)
-	{
-		const FTrajectorySample& SampleA = Trajectory.Samples[StartIdx];
-		const FTrajectorySample& SampleB = Trajectory.Samples[EndIdx];
-		const float TotalTime = SampleB.AccumulatedSeconds - SampleA.AccumulatedSeconds;
-		const float TotalAngleDelta = AccumulatedRotations[EndIdx] - AccumulatedRotations[StartIdx];
-		const float RotationSpeed = TotalAngleDelta / TotalTime;
-		return RotationSpeed;
-	}
+		bool FindTurnBeyondExtrapolation(float MinSharpTurnAngleDegrees)
+		{
+			check(Trajectory.Samples.Num() > 1);
 
-	bool FindTurnBeyondExtrapolation(
-		const FTrajectorySampleRange& Trajectory,
-		const TArrayView<float>& AccumulatedTurnAxisRotations,
-		const TArrayView<float>& ImmediateTurnAxisRotations,
-		float RotSpeedToExtrapolate,
-		float MinSharpTurnAngleRadians)
-	{
-		check(Trajectory.Samples.Num() > 1);
+			const float MinSharpTurnAngleRadians = FMath::DegreesToRadians(MinSharpTurnAngleDegrees);
+
+			const int32 LastSampleIdx = Trajectory.Samples.Num() - 1;
+			const float FirstToLastRotSpeed = CalcRotationSpeed(0, LastSampleIdx);
 
-		const float FirstToLastDeltaSeconds =
-			Trajectory.Samples.Last().AccumulatedSeconds - Trajectory.Samples[0].AccumulatedSeconds;
+			const float SignLastDirRotSpeed = RotSpeedToExtrapolate >= 0 ? 1.0f : -1.0f; // FMath::Sign(RotSpeedToExtrapolate);
+			const float SignedFirstToLastRotSpeed = SignLastDirRotSpeed * FirstToLastRotSpeed;
 
-		const float TargetRotSpeed = MinSharpTurnAngleRadians / FirstToLastDeltaSeconds;
+			// first we'll look for the most aggressive turning frame with respect to the circling speed
+			float FastestRelativeTurnSpeed = 0.0f;
+			int32 FastestRelativeTurnSampleIdx = FindFastestRelativeSpeedSample(FastestRelativeTurnSpeed);
+			if (FastestRelativeTurnSampleIdx == INDEX_NONE)
+			{
+				return false;
+			}
+
+			// Starting from the root sample containing the fastest rotation speed, grow the evaluated interval, from
+			// StartSampleIdx to EndSampleIdx, in a greedy approach. This should lead to the shortest interval that 
+			// contains a large enough turn
+			const FTrajectorySample& RootSample = Trajectory.Samples[FastestRelativeTurnSampleIdx];
+			int32 EndSampleIdx = FastestRelativeTurnSampleIdx;
+			int32 StartSampleIdx = EndSampleIdx - 1;
+			const float TurnSpeedIntervalRadius = RotSpeedToExtrapolate / 2.0f;
+			bool bDone = false;
+			while (!bDone)
+			{
+				const float CurrentTurnLength = 
+					GetAccumulatedDomainValue(Trajectory.Samples[EndSampleIdx]) - 
+					GetAccumulatedDomainValue(Trajectory.Samples[StartSampleIdx]);
+
+				const int32 PrevIdx = StartSampleIdx - 1;
+				float PrevRelativeTurn = 0.0f;
+				bool bPrevValid = IsSampleContributingToSharpTurn(
+					StartSampleIdx,
+					CurrentTurnLength,
+					FastestRelativeTurnSpeed,
+					PrevRelativeTurn);
+
+				const int32 NextIdx = EndSampleIdx + 1;
+				float NextRelativeTurn = 0.0f;
+				bool bNextValid = IsSampleContributingToSharpTurn(
+					NextIdx,
+					CurrentTurnLength,
+					FastestRelativeTurnSpeed,
+					NextRelativeTurn);
+
+				// Every iteration one of 3 things happen:
+				// 1. The interval grows 1 sample into the past
+				// 2. The interval grows 1 sample into the future
+				// 3. The evaluation ends because either a large enough turn was found or there are no more valid
+				//    samples to add to the interval
+				if (bNextValid && bPrevValid)
+				{
+					// if both end points are contributing to the sharp turn, expand the interval towards the biggest
+					// contribution
+					if (FMath::Abs(NextRelativeTurn) >= FMath::Abs(PrevRelativeTurn))
+					{
+						EndSampleIdx = NextIdx;
+					}
+					else
+					{
+						StartSampleIdx = PrevIdx;
+					}
+				}
+				else if (bNextValid)
+				{
+					EndSampleIdx = NextIdx;
+				}
+				else if (bPrevValid)
+				{
+					StartSampleIdx = PrevIdx;
+				}
+				else
+				{
+					// if none of the evaluated extremities contribute to the turn, stop increasing the turn interval.
+					bDone = true;
+				}
+
+				const float ExtrapolatedTurn =
+					RotSpeedToExtrapolate *
+					(Trajectory.Samples[EndSampleIdx].AccumulatedSeconds -
+					 Trajectory.Samples[StartSampleIdx].AccumulatedSeconds);
+				const float TotalTurn =
+					AccumulatedTurnAxisRotations[EndSampleIdx] - AccumulatedTurnAxisRotations[StartSampleIdx];
+
+				if ((SignLastDirRotSpeed * TotalTurn) > (SignLastDirRotSpeed * ExtrapolatedTurn + MinSharpTurnAngleRadians) ||
+					(SignLastDirRotSpeed * TotalTurn) < -MinSharpTurnAngleRadians)
+				{
+					// at any moment if a big enough turn is found, return true, there's no need to do anything else.
+					return true;
+				}
+			}
 
-		const int32 LastSampleIdx = Trajectory.Samples.Num() - 1;
-		const float FirstToLastRotSpeed = CalcRotationSpeed(Trajectory, AccumulatedTurnAxisRotations, 0, LastSampleIdx);
+			return false;
+		}
 
-		const float SignLastDirRotSpeed = FMath::Sign(RotSpeedToExtrapolate);
-		const float SignedFirstToLastRotSpeed = SignLastDirRotSpeed * FirstToLastRotSpeed;
+	private:
 
-		if (SignedFirstToLastRotSpeed > SignLastDirRotSpeed * (RotSpeedToExtrapolate + TargetRotSpeed) ||
-			SignedFirstToLastRotSpeed < -TargetRotSpeed)
+		float CalcRotationSpeed(int32 StartIdx, int32 EndIdx)
 		{
-			// under these conditions there's no need to search further, a turn must exist in the trajectory
-			return true;
+			const FTrajectorySample& SampleA = Trajectory.Samples[StartIdx];
+			const FTrajectorySample& SampleB = Trajectory.Samples[EndIdx];
+			const float TotalTime = SampleB.AccumulatedSeconds - SampleA.AccumulatedSeconds;
+			const float TotalAngleDelta = AccumulatedTurnAxisRotations[EndIdx] - AccumulatedTurnAxisRotations[StartIdx];
+			const float RotationSpeed = TotalAngleDelta / TotalTime;
+			return RotationSpeed;
 		}
 
+		// Populate arrays with turn information. If these end up being used multiple times per frame, we should consider
+		// adding this information to FTrajectorySample (maybe deferred until needed)
+		void CalcTurnData()
+		{
+			AccumulatedTurnAxisRotations.SetNum(Trajectory.Samples.Num());
+			AccumulatedTurnAxisRotations[0] = 0.0f;
 
-		// The case of a turn in the same direction of the circling speed RotSpeedToExtrapolate is already covered by the
-		// conditional above. The code below searches for a turn in the opposite direction only.
+			ImmediateTurnAxisRotations.SetNum(Trajectory.Samples.Num());
+			ImmediateTurnAxisRotations[0] = 0.0f;
 
-		int32 FirstTurnSampleIdx = INDEX_NONE;
-		for (int32 SampleIdx = 0; SampleIdx < LastSampleIdx; ++SampleIdx)
-		{
-			if (SignLastDirRotSpeed * ImmediateTurnAxisRotations[SampleIdx] < 0.0f)
+			const int32 NumSamples = Trajectory.Samples.Num();
+			for (int32 SampleIdx = 1; SampleIdx < NumSamples; ++SampleIdx)
 			{
-				FirstTurnSampleIdx = SampleIdx;
-				break;
+				const FTrajectorySample& PrevSample = Trajectory.Samples[SampleIdx - 1];
+				const FTrajectorySample& CurrSample = Trajectory.Samples[SampleIdx];
+
+				FQuat VelocityDelta = FQuat::FindBetweenVectors(PrevSample.LinearVelocity, CurrSample.LinearVelocity);
+				FVector DeltaAxis;
+				float DeltaAngle = VelocityDelta.GetTwistAngle(TurnAxis);
+
+				ImmediateTurnAxisRotations[SampleIdx] = DeltaAngle;
+				AccumulatedTurnAxisRotations[SampleIdx] = AccumulatedTurnAxisRotations[SampleIdx - 1] + DeltaAngle;
 			}
 		}
 
-		if (FirstTurnSampleIdx == INDEX_NONE)
+		void CalcRotSpeedToExtrapolate()
 		{
-			// no turn to the expected side
-			return false;
+			check(Trajectory.Samples.Num() > 2);
+
+			const int32 LastSampleIdx = Trajectory.Samples.Num() - 1;
+			const FTrajectorySample& LastSample = Trajectory.Samples[LastSampleIdx];
+
+			const int32 BeforeLastSampleIdx = LastSampleIdx - 1;
+			const FTrajectorySample& BeforeLastSample = Trajectory.Samples[BeforeLastSampleIdx];
+
+			const float LastDeltaSeconds = LastSample.AccumulatedSeconds - BeforeLastSample.AccumulatedSeconds;
+
+			const FQuat LastDirectionRotation =
+				FQuat::FindBetweenVectors(BeforeLastSample.LinearVelocity, LastSample.LinearVelocity);
+			const float LastDirRotAngle = LastDirectionRotation.GetTwistAngle(TurnAxis);
+
+			RotSpeedToExtrapolate = LastDirRotAngle / LastDeltaSeconds;
 		}
 
-		int32 LastTurnSampleIdx = FirstTurnSampleIdx + 1;
-		for (int32 SampleIdx = FirstTurnSampleIdx + 2; SampleIdx <= LastSampleIdx; ++SampleIdx)
+		float GetAccumulatedDomainValue(const FTrajectorySample& Sample)
 		{
-			if (SignLastDirRotSpeed * ImmediateTurnAxisRotations[SampleIdx] >= 0.0f)
+			switch (RotationConstraintDomain)
 			{
-				LastTurnSampleIdx = SampleIdx;
-				break;
+			case ETrajectorySampleDomain::Distance:
+				return Sample.AccumulatedDistance;
+			case ETrajectorySampleDomain::Time:
+				return Sample.AccumulatedSeconds;
 			}
+
+			return 0.0f;
 		}
 
-		const float AccumulatedTurnRotation =
-			AccumulatedTurnAxisRotations[LastTurnSampleIdx] - AccumulatedTurnAxisRotations[FirstTurnSampleIdx];
+		// A sample is contributing to the sharp turn if it is turning to the same side as FastestRelativeTurnSpeed,
+		// if adding the sample does not cause the turn to grow larger than the specified constraint, and the turn is
+		// either a. sharper than the extrapolated rotation speed or b. to the opposite side of the extrapolated
+		// rotation speed.
+		bool IsSampleContributingToSharpTurn(
+			int32 SampleIdx,
+			float CurrentTurnLength,
+			float FastestRelativeTurnSpeed, 
+			float& OutSampleRelativeTurn)
+		{
+			OutSampleRelativeTurn = 0.0f;
+
+			bool bIsContributing = false;
 
-		const bool bIsSharpTurn = SignLastDirRotSpeed * AccumulatedTurnRotation <= -MinSharpTurnAngleRadians;
-		return bIsSharpTurn;
-	}
+			if (SampleIdx >= 1 && SampleIdx < Trajectory.Samples.Num())
+			{
+				const int32 PrevIdx = SampleIdx - 1;
+
+				const FTrajectorySample& Sample = Trajectory.Samples[SampleIdx];
+				const float StartDomainValue = GetAccumulatedDomainValue(Sample);
+
+				const FTrajectorySample& PrevSample = Trajectory.Samples[PrevIdx];
+				const float PrevDomainValue = GetAccumulatedDomainValue(PrevSample);
+
+				const float SampleLengthContribution = StartDomainValue - PrevDomainValue;
+				const float NewTurnLength = SampleLengthContribution + CurrentTurnLength;
+
+				if ((RotationConstraintDomain == ETrajectorySampleDomain::None) ||
+					(NewTurnLength <= RotationConstraintLimit))
+				{
+					const float DeltaTime =
+						Trajectory.Samples[SampleIdx].AccumulatedSeconds -
+						Trajectory.Samples[PrevIdx].AccumulatedSeconds;
+					const float PrevTurnSpeed = ImmediateTurnAxisRotations[SampleIdx] / DeltaTime;
+
+					const float TurnSpeedIntervalRadius = RotSpeedToExtrapolate / 2.0f;
+					OutSampleRelativeTurn = PrevTurnSpeed - TurnSpeedIntervalRadius;
+					bIsContributing =
+						FMath::Sign(OutSampleRelativeTurn) == FMath::Sign(FastestRelativeTurnSpeed) &&
+						FMath::Abs(OutSampleRelativeTurn) > FMath::Abs(TurnSpeedIntervalRadius);
+				}
+			}
+
+			return bIsContributing;
+		}
+
+		int32 FindFastestRelativeSpeedSample(float& OutFastestRelativeTurnSpeed)
+		{
+			int32 FastestRelativeTurnSampleIdx = INDEX_NONE;
+			OutFastestRelativeTurnSpeed = 0.0f;
+			float AbsFastestRelativeTurnSpeed = 0.0f;
+
+			// we'll compare turn speeds with the center of the speed interval [0; RotSpeedToInterpolate].
+			// A turn that falls inside the interval will be disregarded.
+			const float TurnSpeedIntervalRadius = RotSpeedToExtrapolate / 2.0f;
+			const int32 NumSamples = Trajectory.Samples.Num();
+			for (int32 SampleIdx = 1; SampleIdx < NumSamples; ++SampleIdx)
+			{
+				const float DeltaTime =
+					Trajectory.Samples[SampleIdx].AccumulatedSeconds -
+					Trajectory.Samples[SampleIdx - 1].AccumulatedSeconds;
+				const float ImmediateSpeed = ImmediateTurnAxisRotations[SampleIdx] / DeltaTime;
+				const float RelativeSpeed = ImmediateSpeed - TurnSpeedIntervalRadius;
+				const float AbsRelativeSpeed = FMath::Abs(RelativeSpeed);
+				if (AbsRelativeSpeed > AbsFastestRelativeTurnSpeed)
+				{
+					FastestRelativeTurnSampleIdx = SampleIdx;
+					OutFastestRelativeTurnSpeed = RelativeSpeed;
+					AbsFastestRelativeTurnSpeed = FMath::Abs(OutFastestRelativeTurnSpeed);
+				}
+			}
+
+			if (AbsFastestRelativeTurnSpeed <= FMath::Abs(TurnSpeedIntervalRadius))
+			{
+				FastestRelativeTurnSampleIdx = INDEX_NONE;
+			}
+
+			return FastestRelativeTurnSampleIdx;
+		}
+
+	};
 } // namespace UE::MotionTrajectory
 
+
 bool UMotionTrajectoryBlueprintLibrary::IsSharpVelocityDirChange(
 	const FTrajectorySampleRange& Trajectory,
 	float MinSharpTurnAngleDegrees,
+	ETrajectorySampleDomain RotationConstraintDomain,
+	float RotationConstraintValue,
 	float MaxAlignmentAngleDegrees,
 	float MinLinearSpeed,
 	FVector TurnAxis,
@@ -510,44 +712,17 @@ bool UMotionTrajectoryBlueprintLibrary::IsSharpVelocityDirChange(
 		return false;
 	}
 
-	const FQuat LastDirectionRotation = 
-		FQuat::FindBetweenVectors(BeforeLastSample.LinearVelocity, LastSample.LinearVelocity);
-	const float LastDirRotAngle = LastDirectionRotation.GetTwistAngle(TurnAxis);
-	const float LastDirRotSpeed = LastDirRotAngle / LastDeltaSeconds;
-
-	TArray<float, TInlineAllocator<128>> AccumulatedTurnAxisRotations;
-	TArray<float, TInlineAllocator<128>> ImmediateTurnAxisRotations;
-	UE::MotionTrajectory::CalcTurnData(Trajectory, TurnAxis, AccumulatedTurnAxisRotations, ImmediateTurnAxisRotations);
-
-	const float PresentToLastRotSpeed = 
-		UE::MotionTrajectory::CalcRotationSpeed(Trajectory, AccumulatedTurnAxisRotations, PresentIdx, LastSampleIdx);
-	const float PresentSpeedDelta = PresentToLastRotSpeed - LastDirRotSpeed;
-	const float PresentAngleDelta = PresentSpeedDelta * LastSample.AccumulatedSeconds;
-
-	const FVector PresentFwd = PresentSample.Transform.GetRotation().RotateVector(ForwardAxis);
-	const FQuat PresentVelToFwdRotation = FQuat::FindBetweenVectors(PresentSample.LinearVelocity, PresentFwd);
-
-	const FVector LastFwd = LastSample.Transform.GetRotation().RotateVector(ForwardAxis);
-	const FQuat LastVelToFwdRotation = FQuat::FindBetweenVectors(LastSample.LinearVelocity, LastFwd);
-
-	const FQuat VelAlignmentDelta = PresentVelToFwdRotation.Inverse() * LastVelToFwdRotation;
-	const float VelAlignmentDeltaAngle = VelAlignmentDelta.GetTwistAngle(TurnAxis);
+	UE::MotionTrajectory::FTurnEvaluator TurnEvaluator = UE::MotionTrajectory::FTurnEvaluator(
+		Trajectory, 
+		RotationConstraintDomain, 
+		RotationConstraintValue, 
+		TurnAxis);
 
-	const float MaxAlignmentAngleRadians = FMath::DegreesToRadians(MaxAlignmentAngleDegrees);
-	if ((FMath::Abs(PresentAngleDelta) < MaxAlignmentAngleRadians) &&
-		(FMath::Abs(VelAlignmentDeltaAngle) < MaxAlignmentAngleRadians))
+	if (!TurnEvaluator.IsTurning(PresentIdx, ForwardAxis, MaxAlignmentAngleDegrees))
 	{
-		// the small difference indicates even if a sharp turn is in the trajectory past, it has already ended
 		return false;
 	}
 
-	const float MinSharpTurnAngleRadians = FMath::DegreesToRadians(MinSharpTurnAngleDegrees);
-	const bool bIsSharpTurn = UE::MotionTrajectory::FindTurnBeyondExtrapolation(
-		Trajectory, 
-		AccumulatedTurnAxisRotations, 
-		ImmediateTurnAxisRotations, 
-		LastDirRotSpeed, 
-		MinSharpTurnAngleRadians);
-
+	const bool bIsSharpTurn = TurnEvaluator.FindTurnBeyondExtrapolation(MinSharpTurnAngleDegrees);
 	return bIsSharpTurn;
 }
diff --git a/Engine/Plugins/Experimental/Animation/MotionTrajectory/Source/MotionTrajectory/Public/MotionTrajectoryLibrary.h b/Engine/Plugins/Experimental/Animation/MotionTrajectory/Source/MotionTrajectory/Public/MotionTrajectoryLibrary.h
--- a/Engine/Plugins/Experimental/Animation/MotionTrajectory/Source/MotionTrajectory/Public/MotionTrajectoryLibrary.h
+++ b/Engine/Plugins/Experimental/Animation/MotionTrajectory/Source/MotionTrajectory/Public/MotionTrajectoryLibrary.h
@@ -152,6 +152,10 @@ public:
 	 * @param Trajectory				Trajectory being evaluated.
 	 * @param MinSharpTurnAngleDegrees	How many degrees of turning must be in the trajectory that can't be explained by 
 	 *									extrapolating the angular velocity at the end of the trajectory.
+	 * @param RotationConstraintDomain	Specifies if the turn must happen unconstrained, within a given distance or
+	 *									within a given time
+	 * @param RotationConstraintValue	If RotationConstraintDomain is not None, this value will specify the constraint
+	 *									threshold
 	 * @param MaxAlignmentAngleDegrees	If the current trajectory sample velocity and facing are well aligned with the 
 	 *									last point in the trajectory future, there's no turn.
 	 * @param MinLinearSpeed			Minimum linear speed at the end points of the trajectory required for this
@@ -165,6 +169,8 @@ public:
 	static bool IsSharpVelocityDirChange(
 		const FTrajectorySampleRange& Trajectory,
 		float MinSharpTurnAngleDegrees = 45.0f,
+		ETrajectorySampleDomain RotationConstraintDomain = ETrajectorySampleDomain::None,
+		float RotationConstraintValue = 0.0f,
 		float MaxAlignmentAngleDegrees = 5.0f,
 		float MinLinearSpeed = 1.0f,
 		FVector TurnAxis = FVector::UpVector,
diff --git a/Engine/Source/Runtime/Engine/Public/Animation/MotionTrajectoryTypes.h b/Engine/Source/Runtime/Engine/Public/Animation/MotionTrajectoryTypes.h
--- a/Engine/Source/Runtime/Engine/Public/Animation/MotionTrajectoryTypes.h
+++ b/Engine/Source/Runtime/Engine/Public/Animation/MotionTrajectoryTypes.h
@@ -10,7 +10,7 @@
 UENUM(BlueprintType, Category="Motion Trajectory", meta=(Bitflags, UseEnumValuesAsMaskValuesInEditor="true"))
 enum class ETrajectorySampleDomain : uint8
 {
-	None = 0 UMETA(Hidden),
+	None = 0,
 	Time = 1 << 0, // Seconds
 	Distance = 1 << 1 // Centimeters (Unreal units)
 };
@@ -93,6 +93,9 @@ struct ENGINE_API FTrajectorySampleRange
 	// Rotates all samples in the trajectory
 	void Rotate(const FQuat& Rotation);
 
+	// Interpolates transform over time
+	void TransformOverTime(const FTransform& Transform, float StartTime, float DeltaTime);
+
 	// Rotates all samples in the trajectory
 	void TransformReferenceFrame(const FTransform& Transform);
 
