Fixed early out for trajectory prediction when rotating on spot.

diff --git a/Engine/Plugins/Experimental/Animation/MotionTrajectory/Source/MotionTrajectory/Private/MotionTrajectoryCharacterMovement.cpp b/Engine/Plugins/Experimental/Animation/MotionTrajectory/Source/MotionTrajectory/Private/MotionTrajectoryCharacterMovement.cpp
--- a/Engine/Plugins/Experimental/Animation/MotionTrajectory/Source/MotionTrajectory/Private/MotionTrajectoryCharacterMovement.cpp
+++ b/Engine/Plugins/Experimental/Animation/MotionTrajectory/Source/MotionTrajectory/Private/MotionTrajectoryCharacterMovement.cpp
@@ -210,7 +210,8 @@ void UCharacterMovementTrajectoryComponent::PredictTrajectory(
 
 			OutTrajectoryRange.Samples.Add(Sample);
 
-			if (FMath::IsNearlyEqual(FMath::Abs(Sample.AccumulatedDistance - PreviousSample.AccumulatedDistance), SMALL_NUMBER))
+			if (FMath::IsNearlyEqual(FMath::Abs(Sample.AccumulatedDistance - PreviousSample.AccumulatedDistance), SMALL_NUMBER) &&
+				Sample.Transform.RotationEquals(PreviousSample.Transform, SMALL_NUMBER))
 			{
 				break;
 			}
@@ -224,8 +225,7 @@ void UCharacterMovementTrajectoryComponent::PredictTrajectory(
 
 			if (((Settings.Domain & TimeDomainMask) == TimeDomainMask)
 				&& (Settings.Seconds > 0.f)
-				&& (Step * IntegrationDelta < Settings.Seconds)
-				&& (Sample.AccumulatedDistance > PreviousSample.AccumulatedDistance))
+				&& (Step * IntegrationDelta < Settings.Seconds))
 			{
 				continue;
 			}
