fix check that would pop up on some platforms due to FP precision
Since we use FMath::IsNearlyZero which checks against UE_SMALL_NUMBER, it's still fine if SampleTime reached the [0,UE_SMALL_NUMBER] interval

diff --git a/Engine/Plugins/Experimental/Animation/MotionTrajectory/Source/MotionTrajectory/Private/MotionTrajectory.cpp b/Engine/Plugins/Experimental/Animation/MotionTrajectory/Source/MotionTrajectory/Private/MotionTrajectory.cpp
--- a/Engine/Plugins/Experimental/Animation/MotionTrajectory/Source/MotionTrajectory/Private/MotionTrajectory.cpp
+++ b/Engine/Plugins/Experimental/Animation/MotionTrajectory/Source/MotionTrajectory/Private/MotionTrajectory.cpp
@@ -250,7 +250,7 @@ FTrajectorySampleRange UMotionTrajectoryComponent::GetHistory() const
 		for (int32 Idx = 0;; ++Idx)
 		{
 			const float SampleTime = FirstUniformSampleTime + static_cast<float>(Idx) * SampleFrequency;
-			check(SampleTime <= 0.f);
+			check(SampleTime <= UE_SMALL_NUMBER);
 
 			if (!FMath::IsNearlyZero(SampleTime))
 			{
