MotionTrajectory: Fixed floating point precision issues in UMotionTrajectoryComponent::GetHistory. Contributed by [REDACTED_PERSON].

diff --git a/Engine/Plugins/Experimental/Animation/MotionTrajectory/Source/MotionTrajectory/Private/MotionTrajectory.cpp b/Engine/Plugins/Experimental/Animation/MotionTrajectory/Source/MotionTrajectory/Private/MotionTrajectory.cpp
--- a/Engine/Plugins/Experimental/Animation/MotionTrajectory/Source/MotionTrajectory/Private/MotionTrajectory.cpp
+++ b/Engine/Plugins/Experimental/Animation/MotionTrajectory/Source/MotionTrajectory/Private/MotionTrajectory.cpp
@@ -240,19 +240,17 @@ FTrajectorySampleRange UMotionTrajectoryComponent::GetHistory() const
 		int32 InterpStartIdx = 0;
 		const float SampleFrequency = 1.f / static_cast<float>(SampleRate);
 		const float FirstNonUniformSampleTime = SampleHistory.First().AccumulatedSeconds;
-		const float FirstUniformSampleTime = FMath::RoundToFloat(FirstNonUniformSampleTime / SampleFrequency) * SampleFrequency;
+		const float FirstUniformSampleTime = FMath::CeilToFloat(FirstNonUniformSampleTime / SampleFrequency) * SampleFrequency;
 
 		check(SampleFrequency > 0.f);
 		check(FirstUniformSampleTime <= 0.f);
 
 		// Iterate the source history buffer with a uniform sample frequency
-		// The iteration exit condition is meant to guarantee that we do not include (erronenously) a instantaneous/present trajectory sample at zero time
+		// The iteration exit condition is meant to guarantee that we do not include (erroneously) a instantaneous/present trajectory sample at zero time
 		for (int32 Idx = 0;; ++Idx)
 		{
 			const float SampleTime = FirstUniformSampleTime + static_cast<float>(Idx) * SampleFrequency;
-			check(SampleTime <= UE_SMALL_NUMBER);
-
-			if (!FMath::IsNearlyZero(SampleTime))
+			if (SampleTime < -UE_KINDA_SMALL_NUMBER)
 			{
 				const FTrajectorySample Sample = FTrajectorySampleRange::IterSampleTrajectory(SampleHistory, ETrajectorySampleDomain::Time, SampleTime, InterpStartIdx, bSmoothInterpolation);
 				UniformHistory.Samples.Emplace(Sample);
