motion matching - removing FMotionTrajectorySettings depenencies from ETrajectorySampleDomain

diff --git a/Engine/Plugins/Experimental/Animation/MotionTrajectory/Source/MotionTrajectory/Private/MotionTrajectory.cpp b/Engine/Plugins/Experimental/Animation/MotionTrajectory/Source/MotionTrajectory/Private/MotionTrajectory.cpp
--- a/Engine/Plugins/Experimental/Animation/MotionTrajectory/Source/MotionTrajectory/Private/MotionTrajectory.cpp
+++ b/Engine/Plugins/Experimental/Animation/MotionTrajectory/Source/MotionTrajectory/Private/MotionTrajectory.cpp
@@ -105,11 +105,6 @@ void UMotionTrajectoryComponent::TickHistoryEvictionPolicy()
 		const FTrajectorySample FirstSample = SampleHistory.First();
 		const float FirstSampleTime = FirstSample.AccumulatedSeconds;
 
-		constexpr int32 DistanceDomainMask = static_cast<int32>(ETrajectorySampleDomain::Distance);
-		constexpr int32 TimeDomainMask = static_cast<int32>(ETrajectorySampleDomain::Time);
-		const bool bDistanceDomainEnabled = (HistorySettings.Domain & DistanceDomainMask) == DistanceDomainMask;
-		const bool bTimeDomainEnabled = (HistorySettings.Domain & TimeDomainMask) == TimeDomainMask;
-
 		// For both time and distance history domains, we compute the concept of an "effective" time horizon:
 		// This value is present to guarantee uniform history decay against a fixed anchor point in the past when there is lack of trajectory motion
 		// The anchor point is defined by the AccumulatedTime of the furthest history sample at the exact moment of zero motion
@@ -136,14 +131,8 @@ void UMotionTrajectoryComponent::TickHistoryEvictionPolicy()
 					return true;
 				}
 
-				// Distance horizon
-				if (bDistanceDomainEnabled && (Sample.AccumulatedDistance < -HistorySettings.Distance))
-				{
-					return true;
-				}
-
 				// Time horizon
-				if (bTimeDomainEnabled && (Sample.AccumulatedSeconds < -HistorySettings.Seconds))
+				if (Sample.AccumulatedSeconds < -HistorySettings.Seconds)
 				{
 					return true;
 				}
diff --git a/Engine/Plugins/Experimental/Animation/MotionTrajectory/Source/MotionTrajectory/Private/MotionTrajectoryCharacterMovement.cpp b/Engine/Plugins/Experimental/Animation/MotionTrajectory/Source/MotionTrajectory/Private/MotionTrajectoryCharacterMovement.cpp
--- a/Engine/Plugins/Experimental/Animation/MotionTrajectory/Source/MotionTrajectory/Private/MotionTrajectoryCharacterMovement.cpp
+++ b/Engine/Plugins/Experimental/Animation/MotionTrajectory/Source/MotionTrajectory/Private/MotionTrajectoryCharacterMovement.cpp
@@ -183,60 +183,31 @@ void UCharacterMovementTrajectoryComponent::PredictTrajectory(
 	OutTrajectoryRange.SampleRate = InSampleRate;
 	const float IntegrationDelta = 1.f / static_cast<float>(InSampleRate);
 
-	if (!!Settings.Domain)
-	{
-		FTrajectorySample Sample = PresentTrajectory;
-		FTrajectorySample PreviousSample = PresentTrajectory;
-		float AccumulatedDistance = 0.f;
-		float AccumulatedSeconds = 0.f;
-		FRotator ControlRotationTotalDelta = FRotator::ZeroRotator;
-
-		constexpr int32 DistanceDomainMask = static_cast<int32>(ETrajectorySampleDomain::Distance);
-		constexpr int32 TimeDomainMask = static_cast<int32>(ETrajectorySampleDomain::Time);
+	FTrajectorySample Sample = PresentTrajectory;
+	FTrajectorySample PreviousSample = PresentTrajectory;
+	float AccumulatedDistance = 0.f;
+	float AccumulatedTime = 0.f;
+	FRotator ControlRotationTotalDelta = FRotator::ZeroRotator;
 
-		for (int32 Step = 0; Step < InMaxSamples; ++Step)
-		{
-			PreviousSample = Sample;
-			StepPrediction(
-				IntegrationDelta, 
-				InDesiredControlRotationVelocity, 
-				ControlRotationTotalDelta, 
-				Sample);
-
-			AccumulatedDistance += 
-				FVector::Distance(PreviousSample.Transform.GetLocation(), Sample.Transform.GetLocation());
-			Sample.AccumulatedDistance = AccumulatedDistance;
-			AccumulatedSeconds += IntegrationDelta;
-			Sample.AccumulatedSeconds = AccumulatedSeconds;
-
-			OutTrajectoryRange.Samples.Add(Sample);
-
-			if (FMath::IsNearlyEqual(FMath::Abs(Sample.AccumulatedDistance - PreviousSample.AccumulatedDistance), SMALL_NUMBER) &&
-				Sample.Transform.RotationEquals(PreviousSample.Transform, SMALL_NUMBER))
-			{
-				break;
-			}
+	for (int32 Step = 0; Step < InMaxSamples; ++Step)
+	{
+		PreviousSample = Sample;
+		StepPrediction(IntegrationDelta, InDesiredControlRotationVelocity, ControlRotationTotalDelta, Sample);
 
-			if (((Settings.Domain & DistanceDomainMask) == DistanceDomainMask)
-				&& (Settings.Distance > 0.f)
-				&& (Sample.AccumulatedDistance < Settings.Distance))
-			{
-				continue;
-			}
+		AccumulatedDistance += FVector::Distance(PreviousSample.Transform.GetLocation(), Sample.Transform.GetLocation());
+		Sample.AccumulatedDistance = AccumulatedDistance;
+		AccumulatedTime += IntegrationDelta;
+		Sample.AccumulatedSeconds = AccumulatedTime;
 
-			if (((Settings.Domain & TimeDomainMask) == TimeDomainMask)
-				&& (Settings.Seconds > 0.f)
-				&& (Step * IntegrationDelta < Settings.Seconds))
-			{
-				continue;
-			}
+		OutTrajectoryRange.Samples.Add(Sample);
 
+		if (Step * IntegrationDelta >= Settings.Seconds)
+		{
 			break;
 		}
 	}
 }
 
-
 // ----------- BEGIN Derived from FCharacterMovementComponentAsyncInput ----------- //
 // FCharacterMovementComponentAsyncInput::CalcVelocity() for linear motion
 // FCharacterMovementComponentAsyncInput::PhysicsRotation() for angular motion
diff --git a/Engine/Plugins/Experimental/Animation/MotionTrajectory/Source/MotionTrajectory/Public/MotionTrajectory.h b/Engine/Plugins/Experimental/Animation/MotionTrajectory/Source/MotionTrajectory/Public/MotionTrajectory.h
--- a/Engine/Plugins/Experimental/Animation/MotionTrajectory/Source/MotionTrajectory/Public/MotionTrajectory.h
+++ b/Engine/Plugins/Experimental/Animation/MotionTrajectory/Source/MotionTrajectory/Public/MotionTrajectory.h
@@ -28,18 +28,9 @@ struct MOTIONTRAJECTORY_API FMotionTrajectorySettings
 {
 	GENERATED_BODY()
 
-	// Trajectory mask for specifying a time and/or distance domain
-	// Default: Time Domain of 2.0 seconds
-	UPROPERTY(EditAnywhere, BlueprintReadWrite, Category="Settings", meta=(Bitmask, BitmaskEnum="/Script/Engine.ETrajectorySampleDomain"))
-	int32 Domain = 1; // ETrajectorySampleDomain::Time
-
 	// Sample time horizon
 	UPROPERTY(EditAnywhere, BlueprintReadWrite, Category="Settings", meta=(ClampMin="0.0"))
 	float Seconds = 2.f;
-
-	// Sample distance horizon
-	UPROPERTY(EditAnywhere, BlueprintReadWrite, Category="Settings", meta=(ClampMin="0.0"))
-	float Distance = 0.f;
 };
 
 // Abstract component interface for the minimum Motion Trajectory prediction and history API
@@ -56,7 +47,7 @@ private:
 	// Previous frame world game time
 	float PreviousWorldGameTime = 0.f;
 
-	// A superceding "effective" time domain applied to distance and time domain history samples, which guarantees uniform decay
+	// A superseding "effective" time domain applied history samples, which guarantees uniform decay
 	float EffectiveTimeDomain = 0.f;
 
 	// Component-driven tick which decays and evicts trajectory history outside of the defined domain horizons
@@ -82,7 +73,7 @@ protected:
 	// Gets the Pawn from the owning Actor of this component
 	const APawn* TryGetOwnerPawn() const;
 
-	// Combines all trajectory samples in the past, present, and future into a unifed trajectory range
+	// Combines all trajectory samples in the past, present, and future into a unified trajectory range
 	FTrajectorySampleRange CombineHistoryPresentPrediction(bool bIncludeHistory, const FTrajectorySampleRange& Prediction) const;
 
 
