motion matching - removing  EPoseSearchFeatureDomain

diff --git a/Engine/Plugins/Experimental/Animation/MotionTrajectory/Source/MotionTrajectory/Private/MotionTrajectory.cpp b/Engine/Plugins/Experimental/Animation/MotionTrajectory/Source/MotionTrajectory/Private/MotionTrajectory.cpp
--- a/Engine/Plugins/Experimental/Animation/MotionTrajectory/Source/MotionTrajectory/Private/MotionTrajectory.cpp
+++ b/Engine/Plugins/Experimental/Animation/MotionTrajectory/Source/MotionTrajectory/Private/MotionTrajectory.cpp
@@ -39,7 +39,6 @@ void UMotionTrajectoryComponent::TickTrajectory(float DeltaTime)
 {
 	// Compute the instantaneous/present trajectory sample and guarantee that is zeroed on all accumulated domains
 	PresentTrajectorySampleWS = CalcWorldSpacePresentTrajectorySample(DeltaTime);
-	PresentTrajectorySampleWS.AccumulatedDistance = 0.f;
 	PresentTrajectorySampleWS.AccumulatedSeconds = 0.f;
 
 	// convert to local space
@@ -243,7 +242,7 @@ FTrajectorySampleRange UMotionTrajectoryComponent::GetHistory() const
 			const float SampleTime = FirstUniformSampleTime + static_cast<float>(Idx) * SampleFrequency;
 			if (SampleTime < -UE_KINDA_SMALL_NUMBER)
 			{
-				const FTrajectorySample Sample = FTrajectorySampleRange::IterSampleTrajectory(SampleHistory, ETrajectorySampleDomain::Time, SampleTime, InterpStartIdx, bSmoothInterpolation);
+				const FTrajectorySample Sample = FTrajectorySampleRange::IterSampleTrajectory(SampleHistory, SampleTime, InterpStartIdx, bSmoothInterpolation);
 				UniformHistory.Samples.Emplace(Sample);
 			}
 			else
diff --git a/Engine/Plugins/Experimental/Animation/MotionTrajectory/Source/MotionTrajectory/Private/MotionTrajectoryCharacterMovement.cpp b/Engine/Plugins/Experimental/Animation/MotionTrajectory/Source/MotionTrajectory/Private/MotionTrajectoryCharacterMovement.cpp
--- a/Engine/Plugins/Experimental/Animation/MotionTrajectory/Source/MotionTrajectory/Private/MotionTrajectoryCharacterMovement.cpp
+++ b/Engine/Plugins/Experimental/Animation/MotionTrajectory/Source/MotionTrajectory/Private/MotionTrajectoryCharacterMovement.cpp
@@ -185,7 +185,6 @@ void UCharacterMovementTrajectoryComponent::PredictTrajectory(
 
 	FTrajectorySample Sample = PresentTrajectory;
 	FTrajectorySample PreviousSample = PresentTrajectory;
-	float AccumulatedDistance = 0.f;
 	float AccumulatedTime = 0.f;
 	FRotator ControlRotationTotalDelta = FRotator::ZeroRotator;
 
@@ -194,8 +193,6 @@ void UCharacterMovementTrajectoryComponent::PredictTrajectory(
 		PreviousSample = Sample;
 		StepPrediction(IntegrationDelta, InDesiredControlRotationVelocity, ControlRotationTotalDelta, Sample);
 
-		AccumulatedDistance += FVector::Distance(PreviousSample.Transform.GetLocation(), Sample.Transform.GetLocation());
-		Sample.AccumulatedDistance = AccumulatedDistance;
 		AccumulatedTime += IntegrationDelta;
 		Sample.AccumulatedSeconds = AccumulatedTime;
 
diff --git a/Engine/Plugins/Experimental/Animation/MotionTrajectory/Source/MotionTrajectory/Private/MotionTrajectoryLibrary.cpp b/Engine/Plugins/Experimental/Animation/MotionTrajectory/Source/MotionTrajectory/Private/MotionTrajectoryLibrary.cpp
--- a/Engine/Plugins/Experimental/Animation/MotionTrajectory/Source/MotionTrajectory/Private/MotionTrajectoryLibrary.cpp
+++ b/Engine/Plugins/Experimental/Animation/MotionTrajectory/Source/MotionTrajectory/Private/MotionTrajectoryLibrary.cpp
@@ -9,44 +9,22 @@
 
 #include UE_INLINE_GENERATED_CPP_BY_NAME(MotionTrajectoryLibrary)
 
-static void FlattenTrajectoryPosition(
-	FTrajectorySample& Sample, 
-	const FTrajectorySample& PrevSample, 
-	const FTrajectorySample& FlattenedPrevSample, 
-	bool PreserveSpeed)
+static void FlattenTrajectoryPosition(FTrajectorySample& Sample, const FTrajectorySample& PrevSample, const FTrajectorySample& FlattenedPrevSample)
 {
 	if (!Sample.Transform.GetTranslation().IsZero())
 	{
 		const FVector Translation = Sample.Transform.GetLocation() - PrevSample.Transform.GetLocation();
 		const FVector FlattenedTranslation = FVector(Translation.X, Translation.Y, 0.0f);
 
-		if (PreserveSpeed)
-		{
-			const float TargetDistance = 
-				FMath::Abs(Sample.AccumulatedDistance - PrevSample.AccumulatedDistance);
-			const FVector FlattenedTranslationDir = Translation.GetSafeNormal2D();
-
-			// Take the full displacement, effectively meaning that the Z axis never existed
-			Sample.Transform.SetLocation(
-				FlattenedPrevSample.Transform.GetLocation() + (FlattenedTranslationDir * TargetDistance));
-		}
-		else
-		{
-			// Accumulate the delta displacement difference as a result of Z axis being removed
-			const float DeltaSeconds = Sample.AccumulatedSeconds - PrevSample.AccumulatedSeconds;
-			const float DeltaDistance = DeltaSeconds >= 0.0f ?
-				FlattenedTranslation.Size() :
-				-FlattenedTranslation.Size();
-
-			Sample.AccumulatedDistance = FlattenedPrevSample.AccumulatedDistance + DeltaDistance;
-			Sample.Transform.SetTranslation(FlattenedPrevSample.Transform.GetLocation() + FlattenedTranslation);
-		}
+		// Accumulate the delta displacement difference as a result of Z axis being removed
+		const float DeltaSeconds = Sample.AccumulatedSeconds - PrevSample.AccumulatedSeconds;
+		const float DeltaDistance = DeltaSeconds >= 0.0f ? FlattenedTranslation.Size() : -FlattenedTranslation.Size();
+		Sample.Transform.SetTranslation(FlattenedPrevSample.Transform.GetLocation() + FlattenedTranslation);
 	}
 }
 
-FTrajectorySampleRange UMotionTrajectoryBlueprintLibrary::FlattenTrajectory2D(FTrajectorySampleRange Trajectory, bool bPreserveSpeed)
+FTrajectorySampleRange UMotionTrajectoryBlueprintLibrary::FlattenTrajectory2D(FTrajectorySampleRange Trajectory)
 {
-
 	if (!Trajectory.HasSamples())
 	{
 		return Trajectory;
@@ -65,7 +43,7 @@ FTrajectorySampleRange UMotionTrajectoryBlueprintLibrary::FlattenTrajectory2D(FT
 		// Linear velocity Z-axis component removal
 		if (!Sample.LinearVelocity.IsZero())
 		{
-			const float VelMagnitude = bPreserveSpeed ? Sample.LinearVelocity.Size() : Sample.LinearVelocity.Size2D();
+			const float VelMagnitude = Sample.LinearVelocity.Size2D();
 			Sample.LinearVelocity = VelMagnitude * Sample.LinearVelocity.GetSafeNormal2D();
 		}
 	}
@@ -87,7 +65,7 @@ FTrajectorySampleRange UMotionTrajectoryBlueprintLibrary::FlattenTrajectory2D(FT
 	for (int32 Idx = PresentSampleIdx + 1, Num = Trajectory.Samples.Num(); Idx < Num; ++Idx)
 	{
 		FTrajectorySample CurrentSample = Trajectory.Samples[Idx];
-		FlattenTrajectoryPosition(Trajectory.Samples[Idx], PrevSample, Trajectory.Samples[Idx - 1], bPreserveSpeed);
+		FlattenTrajectoryPosition(Trajectory.Samples[Idx], PrevSample, Trajectory.Samples[Idx - 1]);
 		PrevSample = CurrentSample;
 	}
 
@@ -102,7 +80,7 @@ FTrajectorySampleRange UMotionTrajectoryBlueprintLibrary::FlattenTrajectory2D(FT
 	for (int32 Idx = PresentSampleIdx - 1, Begin = 0; Idx >= Begin; --Idx)
 	{
 		FTrajectorySample CurrentSample = Trajectory.Samples[Idx];
-		FlattenTrajectoryPosition(Trajectory.Samples[Idx], PrevSample, Trajectory.Samples[Idx + 1], bPreserveSpeed);
+		FlattenTrajectoryPosition(Trajectory.Samples[Idx], PrevSample, Trajectory.Samples[Idx + 1]);
 		PrevSample = CurrentSample;
 	}
 
@@ -140,61 +118,6 @@ static FVector ClampDirection(const FVector InputVector, const TArray<FTrajector
 	return Output;
 }
 
-FTrajectorySampleRange UMotionTrajectoryBlueprintLibrary::ClampTrajectoryDirection(FTrajectorySampleRange Trajectory, const TArray<FTrajectoryDirectionClamp>& Directions, bool bPreserveRotation)
-{
-	if (Directions.IsEmpty())
-	{
-		return Trajectory;
-	}
-
-	if (!Trajectory.HasSamples())
-	{
-		return Trajectory;
-	}
-
-	if (Trajectory.HasOnlyZeroSamples())
-	{
-		return Trajectory;
-	}
-
-	// The clamped present (zero domain) sample is used as the basis for projecting samples along its trajectory
-	FTrajectorySample* PresentSample = Algo::FindByPredicate(Trajectory.Samples, [](const FTrajectorySample& Sample) {
-		return FMath::IsNearlyZero(Sample.AccumulatedSeconds);
-	}
-	);
-
-	check(PresentSample)
-
-	if (!PresentSample->LinearVelocity.IsZero())
-	{
-		const FVector VelocityBasis = ClampDirection(PresentSample->LinearVelocity, Directions).GetSafeNormal();
-
-		for (auto& Sample : Trajectory.Samples)
-		{
-			// Align linear velocity onto the velocity basis to maintain the present intended direction, while retaining per-sample magnitude
-			if (!Sample.LinearVelocity.IsZero())
-			{
-				Sample.LinearVelocity = Sample.LinearVelocity.Size() * Sample.LinearVelocity.ProjectOnTo(VelocityBasis).GetSafeNormal();
-			}
-
-			// Align the position path through projection onto the modified velocity
-			if (!Sample.LinearVelocity.IsZero() && !Sample.Transform.GetLocation().IsZero())
-			{
-				Sample.Transform.SetLocation(
-					FMath::Abs(Sample.AccumulatedDistance) * 
-					Sample.Transform.GetLocation().ProjectOnTo(Sample.LinearVelocity).GetSafeNormal());
-			}
-
-			if (bPreserveRotation)
-			{
-				Sample.Transform.SetRotation(PresentSample->Transform.GetRotation());
-			}
-		}
-	}
-
-	return Trajectory;
-}
-
 FTrajectorySampleRange UMotionTrajectoryBlueprintLibrary::RotateTrajectory(
 	FTrajectorySampleRange Trajectory, 
 	const FQuat& Rotation)
@@ -261,17 +184,11 @@ bool UMotionTrajectoryBlueprintLibrary::IsStoppingTrajectory(
 		const float SquaredLastLinearSpeed = LastLinearVelocity.SquaredLength();
 		
 		int StartIdx = 0;
-		const FTrajectorySample& PresentSample = FTrajectorySampleRange::IterSampleTrajectory(
-			Trajectory.Samples, 
-			ETrajectorySampleDomain::Time, 
-			0.0f, 
-			StartIdx);
+		const FTrajectorySample& PresentSample = FTrajectorySampleRange::IterSampleTrajectory(Trajectory.Samples, 0.f, StartIdx);
 		const FVector PresenLinearVelocity = PresentSample.LinearVelocity;
 		const float SquaredPresentLinearSpeed = PresenLinearVelocity.SquaredLength();
 
-		const bool bIsStopping =
-			(SquaredPresentLinearSpeed >= (MoveMinSpeed * MoveMinSpeed)) &&
-			(SquaredLastLinearSpeed <= (IdleMaxSpeed * IdleMaxSpeed));
+		const bool bIsStopping = (SquaredPresentLinearSpeed >= (MoveMinSpeed * MoveMinSpeed)) && (SquaredLastLinearSpeed <= (IdleMaxSpeed * IdleMaxSpeed));
 
 		return bIsStopping;
 	}
@@ -290,17 +207,11 @@ bool UMotionTrajectoryBlueprintLibrary::IsStartingTrajectory(
 		const float SquaredFirstLinearSpeed = FirstLinearVelocity.SquaredLength();
 
 		int StartIdx = 0;
-		const FTrajectorySample& PresentSample = FTrajectorySampleRange::IterSampleTrajectory(
-			Trajectory.Samples,
-			ETrajectorySampleDomain::Time,
-			0.0f,
-			StartIdx);
+		const FTrajectorySample& PresentSample = FTrajectorySampleRange::IterSampleTrajectory(Trajectory.Samples, 0.f, StartIdx);
 		const FVector PresentLinearVelocity = PresentSample.LinearVelocity;
 		const float SquaredPresentLinearSpeed = PresentLinearVelocity.SquaredLength();
 
-		const bool IsStarting =
-			(SquaredPresentLinearSpeed >= (MoveMinSpeed * MoveMinSpeed)) &&
-			(SquaredFirstLinearSpeed <= (IdleMaxSpeed * IdleMaxSpeed));
+		const bool IsStarting = (SquaredPresentLinearSpeed >= (MoveMinSpeed * MoveMinSpeed)) && (SquaredFirstLinearSpeed <= (IdleMaxSpeed * IdleMaxSpeed));
 
 		return IsStarting;
 	}
@@ -331,11 +242,7 @@ bool UMotionTrajectoryBlueprintLibrary::IsConstantSpeedTrajectory(
 		const bool LastSpeedWithinLimit = IsWithinLimit(SquaredLastLinearSpeed);
 
 		int StartIdx = 0;
-		const FTrajectorySample& PresentSample = FTrajectorySampleRange::IterSampleTrajectory(
-			Trajectory.Samples,
-			ETrajectorySampleDomain::Time,
-			0.0f,
-			StartIdx);
+		const FTrajectorySample& PresentSample = FTrajectorySampleRange::IterSampleTrajectory(Trajectory.Samples, 0.f, StartIdx);
 		const FVector PresentLinearVelocity = PresentSample.LinearVelocity;
 		const float SquaredPresentLinearSpeed = PresentLinearVelocity.SquaredLength();
 		const bool PresentSpeedWithinLimit = IsWithinLimit(SquaredLastLinearSpeed);
@@ -354,13 +261,8 @@ namespace UE::MotionTrajectory
 	{
 
 	public:
-		FTurnEvaluator(
-			const FTrajectorySampleRange& InTrajectory, 
-			ETrajectorySampleDomain InRotationConstraintDomain,
-			float InRotationConstraintLimit,
-			const FVector& InTurnAxis)
+		FTurnEvaluator(const FTrajectorySampleRange& InTrajectory, float InRotationConstraintLimit, const FVector& InTurnAxis)
 			: Trajectory(InTrajectory)
-			, RotationConstraintDomain(InRotationConstraintDomain)
 			, RotationConstraintLimit(InRotationConstraintLimit)
 			, TurnAxis(InTurnAxis)
 		{
@@ -369,7 +271,6 @@ namespace UE::MotionTrajectory
 		}
 
 		const FTrajectorySampleRange& Trajectory;
-		ETrajectorySampleDomain RotationConstraintDomain;
 		float RotationConstraintLimit;
 		FVector TurnAxis;
 
@@ -567,15 +468,7 @@ namespace UE::MotionTrajectory
 
 		float GetAccumulatedDomainValue(const FTrajectorySample& Sample)
 		{
-			switch (RotationConstraintDomain)
-			{
-			case ETrajectorySampleDomain::Distance:
-				return Sample.AccumulatedDistance;
-			case ETrajectorySampleDomain::Time:
-				return Sample.AccumulatedSeconds;
-			}
-
-			return 0.0f;
+			return Sample.AccumulatedSeconds;
 		}
 
 		// A sample is contributing to the sharp turn if it is turning to the same side as FastestRelativeTurnSpeed,
@@ -605,8 +498,7 @@ namespace UE::MotionTrajectory
 				const float SampleLengthContribution = StartDomainValue - PrevDomainValue;
 				const float NewTurnLength = SampleLengthContribution + CurrentTurnLength;
 
-				if ((RotationConstraintDomain == ETrajectorySampleDomain::None) ||
-					(NewTurnLength <= RotationConstraintLimit))
+				if (NewTurnLength <= RotationConstraintLimit)
 				{
 					const float DeltaTime =
 						Trajectory.Samples[SampleIdx].AccumulatedSeconds -
@@ -665,7 +557,6 @@ namespace UE::MotionTrajectory
 bool UMotionTrajectoryBlueprintLibrary::IsSharpVelocityDirChange(
 	const FTrajectorySampleRange& Trajectory,
 	float MinSharpTurnAngleDegrees,
-	ETrajectorySampleDomain RotationConstraintDomain,
 	float RotationConstraintValue,
 	float MaxAlignmentAngleDegrees,
 	float MinLinearSpeed,
@@ -681,11 +572,7 @@ bool UMotionTrajectoryBlueprintLibrary::IsSharpVelocityDirChange(
 	const int32 LastSampleIdx = Trajectory.Samples.Num() - 1;
 
 	int32 PresentIdx = 0;
-	const FTrajectorySample& PresentSample = FTrajectorySampleRange::IterSampleTrajectory(
-		Trajectory.Samples,
-		ETrajectorySampleDomain::Time,
-		0.0f,
-		PresentIdx);
+	const FTrajectorySample& PresentSample = FTrajectorySampleRange::IterSampleTrajectory(Trajectory.Samples, 0.f, PresentIdx);
 
 	if (PresentIdx >= LastSampleIdx)
 	{
@@ -714,11 +601,7 @@ bool UMotionTrajectoryBlueprintLibrary::IsSharpVelocityDirChange(
 		return false;
 	}
 
-	UE::MotionTrajectory::FTurnEvaluator TurnEvaluator = UE::MotionTrajectory::FTurnEvaluator(
-		Trajectory, 
-		RotationConstraintDomain, 
-		RotationConstraintValue, 
-		TurnAxis);
+	UE::MotionTrajectory::FTurnEvaluator TurnEvaluator = UE::MotionTrajectory::FTurnEvaluator(Trajectory, RotationConstraintValue, TurnAxis);
 
 	if (!TurnEvaluator.IsTurning(PresentIdx, ForwardAxis, MaxAlignmentAngleDegrees))
 	{
diff --git a/Engine/Plugins/Experimental/Animation/MotionTrajectory/Source/MotionTrajectory/Public/MotionTrajectoryLibrary.h b/Engine/Plugins/Experimental/Animation/MotionTrajectory/Source/MotionTrajectory/Public/MotionTrajectoryLibrary.h
--- a/Engine/Plugins/Experimental/Animation/MotionTrajectory/Source/MotionTrajectory/Public/MotionTrajectoryLibrary.h
+++ b/Engine/Plugins/Experimental/Animation/MotionTrajectory/Source/MotionTrajectory/Public/MotionTrajectoryLibrary.h
@@ -38,22 +38,7 @@ public:
 	* @return					Z axis flattened, modified trajectory range
 	*/
 	UFUNCTION(BlueprintPure, Category="Motion Trajectory", meta=(BlueprintThreadSafe))
-	static FTrajectorySampleRange FlattenTrajectory2D(FTrajectorySampleRange Trajectory, bool PreserveSpeed = true);
-
-	/**
-	* Projects trajectory samples onto a defined set of allowed directions
-	*
-	* @param Trajectory			Input trajectory range
-	* @param Directions			Input direction clamping, containing angle thresholds for determining source to target direction
-	* @param bPreserveRotation	If true, sample rotations will be replaced with the present sample rotation
-	*
-	* @return					Direction clamped, modified trajectory range
-	*/
-	UFUNCTION(BlueprintPure, Category="Motion Trajectory", meta=(BlueprintThreadSafe, AutoCreateRefTerm="Directions"))
-	static FTrajectorySampleRange ClampTrajectoryDirection(
-		FTrajectorySampleRange Trajectory, 
-		const TArray<FTrajectoryDirectionClamp>& Directions, 
-		bool bPreserveRotation = true);
+	static FTrajectorySampleRange FlattenTrajectory2D(FTrajectorySampleRange Trajectory);
 
 	/**
 	* Rotates the trajectory
@@ -169,7 +154,6 @@ public:
 	static bool IsSharpVelocityDirChange(
 		const FTrajectorySampleRange& Trajectory,
 		float MinSharpTurnAngleDegrees = 45.0f,
-		ETrajectorySampleDomain RotationConstraintDomain = ETrajectorySampleDomain::None,
 		float RotationConstraintValue = 0.0f,
 		float MaxAlignmentAngleDegrees = 5.0f,
 		float MinLinearSpeed = 1.0f,
diff --git a/Engine/Plugins/Experimental/Animation/PoseSearch/Source/Editor/Private/PoseSearchDebugger.h b/Engine/Plugins/Experimental/Animation/PoseSearch/Source/Editor/Private/PoseSearchDebugger.h
--- a/Engine/Plugins/Experimental/Animation/PoseSearch/Source/Editor/Private/PoseSearchDebugger.h
+++ b/Engine/Plugins/Experimental/Animation/PoseSearch/Source/Editor/Private/PoseSearchDebugger.h
@@ -18,7 +18,6 @@ class UAnimSequence;
 class UBlendSpace;
 
 namespace TraceServices { class IAnalysisSession; }
-enum class EPoseSearchFeatureDomain;
 class UPoseSearchDatabase;
 class IUnrealInsightsModule;
 class ITableRow;
diff --git a/Engine/Plugins/Experimental/Animation/PoseSearch/Source/Runtime/Private/PoseSearch.cpp b/Engine/Plugins/Experimental/Animation/PoseSearch/Source/Runtime/Private/PoseSearch.cpp
--- a/Engine/Plugins/Experimental/Animation/PoseSearch/Source/Runtime/Private/PoseSearch.cpp
+++ b/Engine/Plugins/Experimental/Animation/PoseSearch/Source/Runtime/Private/PoseSearch.cpp
@@ -2505,8 +2505,7 @@ void FSequenceBaseSampler::ProcessRootDistance()
 	// Verify we sampled the final frame of the clip
 	check(SampleTime == SequenceBase->GetPlayLength());
 
-	// Also emit root motion summary info to help with sample wrapping in 
-	// FAssetIndexer::GetSampleTimeFromDistance() and FAssetIndexer::GetSampleInfo()
+	// Also emit root motion summary info to help with sample wrapping in FAssetIndexer::GetSampleInfo()
 	TotalRootTransform = LastRootTransform.GetRelativeTransform(InitialRootTransform);
 	TotalRootDistance = AccumulatedRootDistance.Last();
 }
@@ -2896,8 +2895,7 @@ void FBlendSpaceSampler::ProcessRootDistance()
 	// Verify we sampled the final frame of the clip
 	check(SampleTime == PlayLength);
 
-	// Also emit root motion summary info to help with sample wrapping in 
-	// FAssetIndexer::GetSampleTimeFromDistance() and FAssetIndexer::GetSampleInfo()
+	// Also emit root motion summary info to help with sample wrapping in FAssetIndexer::GetSampleInfo()
 	TotalRootTransform = LastRootTransform.GetRelativeTransform(InitialRootTransform);
 	TotalRootDistance = AccumulatedRootDistance.Last(); 
 }
diff --git a/Engine/Plugins/Experimental/Animation/PoseSearch/Source/Runtime/Private/PoseSearchDerivedData.cpp b/Engine/Plugins/Experimental/Animation/PoseSearch/Source/Runtime/Private/PoseSearchDerivedData.cpp
--- a/Engine/Plugins/Experimental/Animation/PoseSearch/Source/Runtime/Private/PoseSearchDerivedData.cpp
+++ b/Engine/Plugins/Experimental/Animation/PoseSearch/Source/Runtime/Private/PoseSearchDerivedData.cpp
@@ -589,7 +589,6 @@ public: // IAssetIndexer
 	const FAssetIndexingContext& GetIndexingContext() const override { return IndexingContext; }
 	FSampleInfo GetSampleInfo(float SampleTime) const override;
 	FSampleInfo GetSampleInfoRelative(float SampleTime, const FSampleInfo& Origin) const override;
-	const float GetSampleTimeFromDistance(float Distance) const override;
 	FTransform MirrorTransform(const FTransform& Transform) const override;
 	FTransform GetTransformAndCacheResults(float SampleTime, float OriginTime, int8 SchemaBoneIdx, bool& Clamped) override;
 
@@ -678,44 +677,6 @@ bool FAssetIndexer::Process()
 	return true;
 }
 
-const float FAssetIndexer::GetSampleTimeFromDistance(float SampleDistance) const
-{
-	constexpr float SMALL_ROOT_DISTANCE = 1.0f;
-	const bool bCanWrap = IndexingContext.AssetSampler->IsLoopable() && IndexingContext.AssetSampler->GetTotalRootDistance() > SMALL_ROOT_DISTANCE;
-
-	float RelativeDistance = SampleDistance;
-	if (SampleDistance < 0.0f && bCanWrap)
-	{
-		// In this case we're sampling a loop backwards, so MainRelativeDistance must adjust so the number of cycles 
-		// is counted correctly.
-		RelativeDistance += IndexingContext.AssetSampler->GetTotalRootDistance();
-	}
-
-	const float MainTotalDistance = IndexingContext.AssetSampler->GetTotalRootDistance();
-	const FSamplingParam SamplingParam = WrapOrClampSamplingParam(bCanWrap, MainTotalDistance, RelativeDistance);
-	const float ClipTime = IndexingContext.AssetSampler->GetTimeFromRootDistance(SamplingParam.WrappedParam + SamplingParam.Extrapolation);
-
-	// Unwrap the clip time
-	float SampleTime = MAX_flt;
-	if (bCanWrap)
-	{
-		if (SampleDistance < 0.0f)
-		{
-			SampleTime = -((SamplingParam.NumCycles * IndexingContext.AssetSampler->GetPlayLength()) + (IndexingContext.AssetSampler->GetPlayLength() - ClipTime));
-		}
-		else
-		{
-			SampleTime = SamplingParam.NumCycles * IndexingContext.AssetSampler->GetPlayLength() + ClipTime;
-		}
-	}
-	else
-	{
-		SampleTime = ClipTime;
-	}
-
-	return SampleTime;
-}
-
 FAssetIndexer::FSampleInfo FAssetIndexer::GetSampleInfo(float SampleTime) const
 {
 	FSampleInfo Sample;
diff --git a/Engine/Plugins/Experimental/Animation/PoseSearch/Source/Runtime/Private/PoseSearchFeatureChannels.cpp b/Engine/Plugins/Experimental/Animation/PoseSearch/Source/Runtime/Private/PoseSearchFeatureChannels.cpp
--- a/Engine/Plugins/Experimental/Animation/PoseSearch/Source/Runtime/Private/PoseSearchFeatureChannels.cpp
+++ b/Engine/Plugins/Experimental/Animation/PoseSearch/Source/Runtime/Private/PoseSearchFeatureChannels.cpp
@@ -1401,23 +1401,6 @@ void UPoseSearchFeatureChannel_Trajectory::IndexAsset(UE::PoseSearch::IAssetInde
 	}
 }
 
-float UPoseSearchFeatureChannel_Trajectory::GetSampleTime(const UE::PoseSearch::IAssetIndexer& Indexer, float Offset, float SampleTime, float RootDistance) const
-{
-	switch (Domain)
-	{
-	case EPoseSearchFeatureDomain::Time:
-		return SampleTime + Offset;
-
-	case EPoseSearchFeatureDomain::Distance:
-		return Indexer.GetSampleTimeFromDistance(RootDistance + Offset);
-
-	default:
-		checkNoEntry();
-	}
-	
-	return 0.0f;
-}
-
 void UPoseSearchFeatureChannel_Trajectory::IndexAssetPrivate(const UE::PoseSearch::IAssetIndexer& Indexer, int32 SampleIdx, TArrayView<float> FeatureVector) const
 {
 	// This function samples the instantaneous trajectory at time t as well as the trajectory's velocity and acceleration at time t.
@@ -1437,7 +1420,7 @@ void UPoseSearchFeatureChannel_Trajectory::IndexAssetPrivate(const UE::PoseSearc
 	int32 DataOffset = ChannelDataOffset;
 	for (const FPoseSearchTrajectorySample& Sample : Samples)
 	{
-		const float SubsampleTime = GetSampleTime(Indexer, Sample.Offset, SampleTime, Origin.RootDistance);
+		const float SubsampleTime = Sample.Offset + SampleTime;
 
 		// For each pose subsample term, get the corresponding clip, accumulated root motion,
 			// and wrap the time parameter based on the clip's length.
@@ -1520,22 +1503,6 @@ bool UPoseSearchFeatureChannel_Trajectory::BuildQuery(UE::PoseSearch::FSearchCon
 		return false;
 	}
 
-	ETrajectorySampleDomain SampleDomain;
-	switch (Domain)
-	{
-		case EPoseSearchFeatureDomain::Time:
-			SampleDomain = ETrajectorySampleDomain::Time;
-			break;
-
-		case EPoseSearchFeatureDomain::Distance:
-			SampleDomain = ETrajectorySampleDomain::Distance;
-			break;
-
-		default:
-			checkNoEntry();
-			return false;
-	}
-
 	int32 NextIterStartIdx = 0;
 	int32 DataOffset = ChannelDataOffset;
 	float PreviousOffset = -FLT_MAX;
@@ -1543,7 +1510,7 @@ bool UPoseSearchFeatureChannel_Trajectory::BuildQuery(UE::PoseSearch::FSearchCon
 	{
 		// making sure Samples are sorted
 		check(Sample.Offset >= PreviousOffset);
-		const FTrajectorySample TrajectorySample = FTrajectorySampleRange::IterSampleTrajectory(SearchContext.Trajectory->Samples, SampleDomain, Sample.Offset, NextIterStartIdx);
+		const FTrajectorySample TrajectorySample = FTrajectorySampleRange::IterSampleTrajectory(SearchContext.Trajectory->Samples, Sample.Offset, NextIterStartIdx);
 
 		const FVector LinearVelocityDirection = TrajectorySample.LinearVelocity.GetClampedToMaxSize(1.0f);
 		const FVector FacingDirection = TrajectorySample.Transform.GetRotation().GetForwardVector();
diff --git a/Engine/Plugins/Experimental/Animation/PoseSearch/Source/Runtime/Private/PoseSearchLibrary.cpp b/Engine/Plugins/Experimental/Animation/PoseSearch/Source/Runtime/Private/PoseSearchLibrary.cpp
--- a/Engine/Plugins/Experimental/Animation/PoseSearch/Source/Runtime/Private/PoseSearchLibrary.cpp
+++ b/Engine/Plugins/Experimental/Animation/PoseSearch/Source/Runtime/Private/PoseSearchLibrary.cpp
@@ -233,15 +233,8 @@ static void TraceMotionMatchingState(
 	{
 		// simulation
 		int32 FirstIdx = 0;
-		const FTrajectorySample PrevSample = FTrajectorySampleRange::IterSampleTrajectory(
-			Trajectory.Samples,
-			ETrajectorySampleDomain::Time,
-			-DeltaTime, FirstIdx);
-
-		const FTrajectorySample CurrSample = FTrajectorySampleRange::IterSampleTrajectory(
-			Trajectory.Samples,
-			ETrajectorySampleDomain::Time,
-			0.0f, FirstIdx);
+		const FTrajectorySample PrevSample = FTrajectorySampleRange::IterSampleTrajectory(Trajectory.Samples, -DeltaTime, FirstIdx);
+		const FTrajectorySample CurrSample = FTrajectorySampleRange::IterSampleTrajectory(Trajectory.Samples, 0.f, FirstIdx);
 
 		const FTransform SimDelta = CurrSample.Transform.GetRelativeTransform(PrevSample.Transform);
 
diff --git a/Engine/Plugins/Experimental/Animation/PoseSearch/Source/Runtime/Public/PoseSearch/PoseSearch.h b/Engine/Plugins/Experimental/Animation/PoseSearch/Source/Runtime/Public/PoseSearch/PoseSearch.h
--- a/Engine/Plugins/Experimental/Animation/PoseSearch/Source/Runtime/Public/PoseSearch/PoseSearch.h
+++ b/Engine/Plugins/Experimental/Animation/PoseSearch/Source/Runtime/Public/PoseSearch/PoseSearch.h
@@ -423,7 +423,6 @@ public:
 	virtual const FAssetIndexingContext& GetIndexingContext() const = 0;
 	virtual FSampleInfo GetSampleInfo(float SampleTime) const = 0;
 	virtual FSampleInfo GetSampleInfoRelative(float SampleTime, const FSampleInfo& Origin) const = 0;
-	virtual const float GetSampleTimeFromDistance(float Distance) const = 0;
 	virtual FTransform MirrorTransform(const FTransform& Transform) const = 0;
 	virtual FTransform GetTransformAndCacheResults(float SampleTime, float OriginTime, int8 SchemaBoneIdx, bool& Clamped) = 0;
 };
diff --git a/Engine/Plugins/Experimental/Animation/PoseSearch/Source/Runtime/Public/PoseSearch/PoseSearchFeatureChannels.h b/Engine/Plugins/Experimental/Animation/PoseSearch/Source/Runtime/Public/PoseSearch/PoseSearchFeatureChannels.h
--- a/Engine/Plugins/Experimental/Animation/PoseSearch/Source/Runtime/Public/PoseSearch/PoseSearchFeatureChannels.h
+++ b/Engine/Plugins/Experimental/Animation/PoseSearch/Source/Runtime/Public/PoseSearch/PoseSearchFeatureChannels.h
@@ -194,16 +194,6 @@ protected:
 
 //////////////////////////////////////////////////////////////////////////
 // UPoseSearchFeatureChannel_Trajectory
-UENUM()
-enum class EPoseSearchFeatureDomain : int32
-{
-	Time,
-	Distance,
-
-	Num UMETA(Hidden),
-	Invalid = Num UMETA(Hidden)
-};
-
 UENUM(meta = (Bitflags, UseEnumValuesAsMaskValuesInEditor = "true"))
 enum class EPoseSearchTrajectoryFlags : uint32
 {
@@ -244,9 +234,6 @@ class POSESEARCH_API UPoseSearchFeatureChannel_Trajectory : public UPoseSearchFe
 	GENERATED_BODY()
 
 public:
-	UPROPERTY(EditAnywhere, Category="Settings")
-	EPoseSearchFeatureDomain Domain = EPoseSearchFeatureDomain::Time;
-
 	UPROPERTY(EditAnywhere, Category = "Settings")
 	float Weight = 1.f;
 
@@ -272,7 +259,6 @@ public:
 
 protected:
 	void IndexAssetPrivate(const UE::PoseSearch::IAssetIndexer& Indexer, int32 SampleIdx, TArrayView<float> FeatureVector) const;
-	float GetSampleTime(const UE::PoseSearch::IAssetIndexer& Indexer, float Offset, float SampleTime, float RootDistance) const;
 };
 
 //////////////////////////////////////////////////////////////////////////
diff --git a/Engine/Source/Runtime/Engine/Private/Animation/MotionTrajectoryTypes.cpp b/Engine/Source/Runtime/Engine/Private/Animation/MotionTrajectoryTypes.cpp
--- a/Engine/Source/Runtime/Engine/Private/Animation/MotionTrajectoryTypes.cpp
+++ b/Engine/Source/Runtime/Engine/Private/Animation/MotionTrajectoryTypes.cpp
@@ -14,17 +14,14 @@
 static constexpr int32 DebugTrajectorySampleDisable = 0;
 static constexpr int32 DebugTrajectorySampleCount = 1;
 static constexpr int32 DebugTrajectorySampleTime = 2;
-static constexpr int32 DebugTrajectorySampleDistance = 3;
-static constexpr int32 DebugTrajectorySamplePosition = 4;
-static constexpr int32 DebugTrajectorySampleVelocity = 5;
+static constexpr int32 DebugTrajectorySamplePosition = 3;
+static constexpr int32 DebugTrajectorySampleVelocity = 4;
 static const FVector DebugSampleTypeOffset(0.f, 0.f, 50.f);
 static const FVector DebugSampleOffset(0.f, 0.f, 10.f);
 
 TAutoConsoleVariable<int32> CVarMotionTrajectoryDebug(TEXT("a.MotionTrajectory.Debug"), 0, TEXT("Turn on debug drawing for motion trajectory"));
 TAutoConsoleVariable<int32> CVarMotionTrajectoryDebugStride(TEXT("a.MotionTrajectory.Stride"), 1, TEXT("Configure the sample stride when displaying information"));
-TAutoConsoleVariable<int32> CVarMotionTrajectoryDebugOptions(TEXT("a.MotionTrajectory.Options"), 0,
-	TEXT("Toggle motion trajectory sample information:\n 0. Disable Text\n 1. Index\n2. Accumulated Time\n 3. Accumulated Distance\n 4. Position\n 5. Velocity\n 6. Acceleration")
-);
+TAutoConsoleVariable<int32> CVarMotionTrajectoryDebugOptions(TEXT("a.MotionTrajectory.Options"), 0, TEXT("Toggle motion trajectory sample information:\n 0. Disable Text\n 1. Index\n2. Accumulated Time\n 3. Position\n 4. Velocity\n 5. Acceleration"));
 #endif
 
 namespace
@@ -62,7 +59,6 @@ bool FTrajectorySample::IsZeroSample() const
 	// AccumulatedTime is specifically omitted here to allow for the zero sample semantic across an entire trajectory range
 	return LinearVelocity.IsNearlyZero()
 		&& Transform.GetTranslation().IsNearlyZero()
-		&& FMath::IsNearlyZero(AccumulatedDistance)
 		&& Transform.GetRotation().IsIdentity();
 }
 
@@ -70,7 +66,6 @@ FTrajectorySample FTrajectorySample::Lerp(const FTrajectorySample& Sample, float
 {
 	FTrajectorySample Interp;
 	Interp.AccumulatedSeconds = FMath::Lerp(AccumulatedSeconds, Sample.AccumulatedSeconds, Alpha);
-	Interp.AccumulatedDistance = FMath::Lerp(AccumulatedDistance, Sample.AccumulatedDistance, Alpha);
 	Interp.LinearVelocity = FMath::Lerp(LinearVelocity, Sample.LinearVelocity, Alpha);
 
 	Interp.Transform.Blend(Transform, Sample.Transform, Alpha);
@@ -84,7 +79,6 @@ FTrajectorySample FTrajectorySample::SmoothInterp(const FTrajectorySample& PrevS
 	, float Alpha) const
 {
 	FTrajectorySample Interp;
-	Interp.AccumulatedDistance = CubicCRSplineInterpSafe(PrevSample.AccumulatedDistance, AccumulatedDistance, Sample.AccumulatedDistance, NextSample.AccumulatedDistance, Alpha);
 	Interp.AccumulatedSeconds = CubicCRSplineInterpSafe(PrevSample.AccumulatedSeconds, AccumulatedSeconds, Sample.AccumulatedSeconds, NextSample.AccumulatedSeconds, Alpha);
 	Interp.LinearVelocity = CubicCRSplineInterpSafe(PrevSample.LinearVelocity, LinearVelocity, Sample.LinearVelocity, NextSample.LinearVelocity, Alpha);
 
@@ -115,20 +109,6 @@ FTrajectorySample FTrajectorySample::SmoothInterp(const FTrajectorySample& PrevS
 void FTrajectorySample::PrependOffset(const FTransform DeltaTransform, float DeltaSeconds)
 {
 	AccumulatedSeconds += DeltaSeconds;
-
-	if (FMath::IsNearlyZero(AccumulatedSeconds))
-	{
-		AccumulatedDistance = 0.0f;
-	}
-	else
-	{
-		const float DistanceOffset = DeltaSeconds >= 0.0f ?
-			DeltaTransform.GetTranslation().Size() :
-			-DeltaTransform.GetTranslation().Size();
-
-		AccumulatedDistance += DistanceOffset;
-	}
-
 	Transform *= DeltaTransform;
 	LinearVelocity = DeltaTransform.TransformVectorNoScale(LinearVelocity);
 }
@@ -245,10 +225,6 @@ void FTrajectorySampleRange::DebugDrawTrajectory(bool bEnable
 					DebugString = "Sample Time:";
 					DebugSampleString = DebugSampleString.Format(TEXT("{0}"), { Samples[Idx].AccumulatedSeconds });
 					break;
-				case DebugTrajectorySampleDistance: // Sample Accumulated Distance
-					DebugString = "Sample Distance:";
-					DebugSampleString = DebugSampleString.Format(TEXT("{0}"), { Samples[Idx].AccumulatedDistance });
-					break;
 				case DebugTrajectorySamplePosition: // Sample Position
 					DebugString = "Sample Position:";
 					DebugSampleString = DebugSampleString.Format(TEXT("{0}"), { Samples[Idx].Transform.GetLocation().ToCompactString() });
diff --git a/Engine/Source/Runtime/Engine/Public/Animation/MotionTrajectoryTypes.h b/Engine/Source/Runtime/Engine/Public/Animation/MotionTrajectoryTypes.h
--- a/Engine/Source/Runtime/Engine/Public/Animation/MotionTrajectoryTypes.h
+++ b/Engine/Source/Runtime/Engine/Public/Animation/MotionTrajectoryTypes.h
@@ -6,17 +6,6 @@
 
 #include "MotionTrajectoryTypes.generated.h"
 
-// Enumeration for signaling which "Accumulated" domain values to respect when determining past and future sampling horizons
-UENUM(BlueprintType, Category="Motion Trajectory", meta=(Bitflags, UseEnumValuesAsMaskValuesInEditor="true"))
-enum class ETrajectorySampleDomain : uint8
-{
-	None = 0,
-	Time = 1 << 0, // Seconds
-	Distance = 1 << 1 // Centimeters (Unreal units)
-};
-ENUM_CLASS_FLAGS(ETrajectorySampleDomain);
-
-// A motion trajectory sample associated with a specific time or distance domain value
 USTRUCT(BlueprintType, Category="Motion Trajectory")
 struct ENGINE_API FTrajectorySample
 {
@@ -27,11 +16,6 @@ struct ENGINE_API FTrajectorySample
 	UPROPERTY(EditAnywhere, BlueprintReadWrite, Category="Motion Trajectory")
 	float AccumulatedSeconds = 0.f;
 
-	// The relative accumulated distance that this trajectory sample is associated with
-	// Zero value for instantaneous, negative values for the past, and positive values for the future
-	UPROPERTY(EditAnywhere, BlueprintReadWrite, Category="Motion Trajectory")
-	float AccumulatedDistance = 0.f;
-
 	// Position relative to the sampled in-motion object
 	UPROPERTY(EditAnywhere, BlueprintReadWrite, Category="Motion Trajectory")
 	FTransform Transform = FTransform::Identity;
@@ -74,7 +58,7 @@ struct ENGINE_API FTrajectorySampleRange
 	static constexpr float DebugDefaultArrowSize = 40.f;
 	static constexpr float DebugDefaultArrowThickness = 2.f;
 
-	// Linearly ordered (Time or Distance domain) container for past, present, and future trajectory samples
+	// Linearly ordered container for past, present, and future trajectory samples
 	UPROPERTY(EditAnywhere, BlueprintReadWrite, Category="Motion Trajectory")
 	TArray<FTrajectorySample> Samples;
 
@@ -118,7 +102,7 @@ struct ENGINE_API FTrajectorySampleRange
 		, float VelArrowThickness = DebugDefaultArrowThickness) const;
 
 	// Iterator for precise subsampling of the trajectory
-	template <typename Container> static FTrajectorySample IterSampleTrajectory(const Container& Samples, ETrajectorySampleDomain DomainType, float DomainValue, int32& InitialIdx, bool bSmoothInterp = false)
+	template <typename Container> static FTrajectorySample IterSampleTrajectory(const Container& Samples, float DomainValue, int32& InitialIdx, bool bSmoothInterp = false)
 	{
 		auto SafeNextIndex = [](float Idx, float Num)
 		{
@@ -134,7 +118,7 @@ struct ENGINE_API FTrajectorySampleRange
 		for (int32 Idx = InitialIdx, Num = Samples.Num(); Idx < Num; ++Idx)
 		{
 			const FTrajectorySample& NextSample = Samples[Idx];
-			const float NextDomainValue = DomainType == ETrajectorySampleDomain::Time ? NextSample.AccumulatedSeconds : NextSample.AccumulatedDistance;
+			const float NextDomainValue = NextSample.AccumulatedSeconds;
 
 			// Continue traversing the samples until an appropriate "right-hand" side domain value is found for interpolation 
 			if (NextDomainValue < DomainValue)
@@ -149,7 +133,7 @@ struct ENGINE_API FTrajectorySampleRange
 				// Find the interpolation factor of P between: [P1 .. P .. P2] by mapping the domain values to [0 ... 1]
 				// Note: NextIdx is biased to equal P1 rather than P2 to account for cases where subsequent trajectory iterations may require subsampling between [P1 ... P2] again
 				const FTrajectorySample& InitialSample = Samples[InitialIdx];
-				const float PrevDomainValue = DomainType == ETrajectorySampleDomain::Time ? InitialSample.AccumulatedSeconds : InitialSample.AccumulatedDistance;
+				const float PrevDomainValue = InitialSample.AccumulatedSeconds;
 				const float Alpha = FMath::GetMappedRangeValueUnclamped(FVector2f(PrevDomainValue, NextDomainValue), FVector2f(0.f, 1.f), DomainValue);
 
 				FTrajectorySample InterpSample;
