motion matching - fix for "warning V674: The '1.e-4f' literal of the 'float' type is compared to a value of the 'int' type. Consider inspecting the 'InSampleRate > (1.e-4f)' expression."

diff --git a/Engine/Plugins/Experimental/Animation/MotionTrajectory/Source/MotionTrajectory/Private/MotionTrajectoryCharacterMovement.cpp b/Engine/Plugins/Experimental/Animation/MotionTrajectory/Source/MotionTrajectory/Private/MotionTrajectoryCharacterMovement.cpp
--- a/Engine/Plugins/Experimental/Animation/MotionTrajectory/Source/MotionTrajectory/Private/MotionTrajectoryCharacterMovement.cpp
+++ b/Engine/Plugins/Experimental/Animation/MotionTrajectory/Source/MotionTrajectory/Private/MotionTrajectoryCharacterMovement.cpp
@@ -166,7 +166,7 @@ void UCharacterMovementTrajectoryComponent::PredictTrajectory(
 	const FRotator& InDesiredControlRotationVelocity,
 	FTrajectorySampleRange& OutTrajectoryRange) const
 {
-	check(InSampleRate > UE_KINDA_SMALL_NUMBER);
+	check(InSampleRate > 0);
 	const float IntegrationDelta = 1.f / static_cast<float>(InSampleRate);
 
 	FTrajectorySample Sample = PresentTrajectory;
