Initial submit of new trajectory code. This just a first step.
This CL brings the trajectory query data into the PoseSearch plugin. We intend to deprecate and eventually remove the data that's defined in the engine. This is mostly to keep the data and API specific to PoseSearch and reduce the number of systems that will be impacted by future changes.
Created UCharacterTrajectoryComponent, which generates trajectories that are specific to ACharacter. The trajectories can be fed directly to the Motion Matching node, removing the need to do extra conversion steps in Blueprint.

diff --git a/Engine/Plugins/Experimental/Animation/MotionTrajectory/MotionTrajectory.uplugin b/Engine/Plugins/Experimental/Animation/MotionTrajectory/MotionTrajectory.uplugin
--- a/Engine/Plugins/Experimental/Animation/MotionTrajectory/MotionTrajectory.uplugin
+++ b/Engine/Plugins/Experimental/Animation/MotionTrajectory/MotionTrajectory.uplugin
@@ -21,5 +21,11 @@
 			"Type": "Runtime",
 			"LoadingPhase": "PreDefault"
 		}
+	],
+	"Plugins": [
+		{
+			"Name": "PoseSearch",
+			"Enabled": true
+		}
 	]
 }
\ No newline at end of file
diff --git a/Engine/Plugins/Experimental/Animation/MotionTrajectory/Source/MotionTrajectory/MotionTrajectory.Build.cs b/Engine/Plugins/Experimental/Animation/MotionTrajectory/Source/MotionTrajectory/MotionTrajectory.Build.cs
--- a/Engine/Plugins/Experimental/Animation/MotionTrajectory/Source/MotionTrajectory/MotionTrajectory.Build.cs
+++ b/Engine/Plugins/Experimental/Animation/MotionTrajectory/Source/MotionTrajectory/MotionTrajectory.Build.cs
@@ -14,7 +14,8 @@ public class MotionTrajectory : ModuleRules
 				"Core",
 				"CoreUObject",
                 "Engine",
-				"AnimGraphRuntime"
+				"AnimGraphRuntime",
+				"PoseSearch"
 			}
 		);
 	}
diff --git a/Engine/Plugins/Experimental/Animation/MotionTrajectory/Source/MotionTrajectory/Private/CharacterMovementTrajectoryLibrary.cpp b/Engine/Plugins/Experimental/Animation/MotionTrajectory/Source/MotionTrajectory/Private/CharacterMovementTrajectoryLibrary.cpp
new file mode 100644
--- /dev/null
+++ b/Engine/Plugins/Experimental/Animation/MotionTrajectory/Source/MotionTrajectory/Private/CharacterMovementTrajectoryLibrary.cpp
@@ -0,0 +1,86 @@
+// Copyright Epic Games, Inc. All Rights Reserved.
+
+#include "CharacterMovementTrajectoryLibrary.h"
+
+#include "GameFramework/CharacterMovementComponent.h"
+#include "MotionTrajectory.h"
+
+#include UE_INLINE_GENERATED_CPP_BY_NAME(CharacterMovementTrajectoryLibrary)
+
+void UCharacterMovementTrajectoryLibrary::StepCharacterMovementGroundPrediction(
+	float InDeltaTime, const FVector& InVelocity, const FVector& InAcceleration, const UCharacterMovementComponent* InCharacterMovementComponent,
+	FVector& OutVelocity)
+{
+	OutVelocity = InVelocity;
+
+	if (InCharacterMovementComponent == nullptr)
+	{
+		UE_LOG(LogMotionTrajectory, Warning, TEXT("UCharacterMovementTrajectoryLibrary::StepCharacterMovementGroundPrediction requires a valid CharacterMovementComponent"));
+		return;
+	}
+
+	// Braking logic is copied from UCharacterMovementComponent::ApplyVelocityBraking()
+	if (InAcceleration.IsZero())
+	{
+		if (InVelocity.IsZero())
+		{
+			return;
+		}
+
+		float ActualBrakingFriction = InCharacterMovementComponent->bUseSeparateBrakingFriction ? InCharacterMovementComponent->BrakingFriction : InCharacterMovementComponent->GroundFriction;
+		const float FrictionFactor = FMath::Max(0.f, InCharacterMovementComponent->BrakingFrictionFactor);
+		ActualBrakingFriction = FMath::Max(0.f, ActualBrakingFriction * FrictionFactor);
+
+		const float BrakingDeceleration = FMath::Max(0.f, InCharacterMovementComponent->GetMaxBrakingDeceleration());
+		const bool bZeroFriction = (ActualBrakingFriction == 0.f);
+		const bool bZeroBraking = (BrakingDeceleration == 0.f);
+
+		if (bZeroFriction && bZeroBraking)
+		{
+			return;
+		}
+
+		float RemainingTime = InDeltaTime;
+		const float MaxTimeStep = FMath::Clamp(InCharacterMovementComponent->BrakingSubStepTime, 1.0f / 75.0f, 1.0f / 20.0f);
+
+		const FVector PrevLinearVelocity = OutVelocity;
+		const FVector RevAccel = (bZeroBraking ? FVector::ZeroVector : (-BrakingDeceleration * OutVelocity.GetSafeNormal()));
+
+		// Decelerate to brake to a stop
+		while (RemainingTime >= UCharacterMovementComponent::MIN_TICK_TIME)
+		{
+			// Zero friction uses constant deceleration, so no need for iteration.
+			const float dt = ((RemainingTime > MaxTimeStep && !bZeroFriction) ? FMath::Min(MaxTimeStep, RemainingTime * 0.5f) : RemainingTime);
+			RemainingTime -= dt;
+
+			// apply friction and braking
+			OutVelocity = OutVelocity + ((-ActualBrakingFriction) * OutVelocity + RevAccel) * dt;
+
+			// Don't reverse direction
+			if ((OutVelocity | PrevLinearVelocity) <= 0.f)
+			{
+				OutVelocity = FVector::ZeroVector;
+				return;
+			}
+		}
+
+		// Clamp to zero if nearly zero, or if below min threshold and braking
+		const float VSizeSq = OutVelocity.SizeSquared();
+		if (VSizeSq <= KINDA_SMALL_NUMBER || (!bZeroBraking && VSizeSq <= FMath::Square(UCharacterMovementComponent::BRAKE_TO_STOP_VELOCITY)))
+		{
+			OutVelocity = FVector::ZeroVector;
+		}
+	}
+	// Acceleration logic is copied from  UCharacterMovementComponent::CalcVelocity
+	else
+	{
+		const FVector AccelDir = InAcceleration.GetSafeNormal();
+		const float VelSize = OutVelocity.Size();
+
+		OutVelocity = OutVelocity - (OutVelocity - AccelDir * VelSize) * FMath::Min(InDeltaTime * InCharacterMovementComponent->GroundFriction, 1.f);
+
+		const float MaxInputSpeed = FMath::Max(InCharacterMovementComponent->GetMaxSpeed() * InCharacterMovementComponent->GetAnalogInputModifier(), InCharacterMovementComponent->GetMinAnalogSpeed());
+		OutVelocity += InAcceleration * InDeltaTime;
+		OutVelocity = OutVelocity.GetClampedToMaxSize(MaxInputSpeed);
+	}
+}
\ No newline at end of file
diff --git a/Engine/Plugins/Experimental/Animation/MotionTrajectory/Source/MotionTrajectory/Private/CharacterTrajectoryComponent.cpp b/Engine/Plugins/Experimental/Animation/MotionTrajectory/Source/MotionTrajectory/Private/CharacterTrajectoryComponent.cpp
new file mode 100644
--- /dev/null
+++ b/Engine/Plugins/Experimental/Animation/MotionTrajectory/Source/MotionTrajectory/Private/CharacterTrajectoryComponent.cpp
@@ -0,0 +1,188 @@
+// Copyright Epic Games, Inc. All Rights Reserved.
+
+#include "CharacterTrajectoryComponent.h"
+
+#include "CharacterMovementTrajectoryLibrary.h"
+#include "Components/SkeletalMeshComponent.h"
+#include "GameFramework/Character.h"
+#include "GameFramework/CharacterMovementComponent.h"
+#include "HAL/IConsoleManager.h"
+#include "MotionTrajectory.h"
+
+#include "MotionTrajectoryCharacterMovement.h"
+#include "MotionTrajectoryLibrary.h"
+
+#if ENABLE_ANIM_DEBUG
+TAutoConsoleVariable<int32> CVarCharacterTrajectoryDebug(TEXT("a.CharacterTrajectory.Debug"), 0, TEXT("Turn on debug drawing for Character trajectory"));
+#endif // ENABLE_ANIM_DEBUG
+
+UCharacterTrajectoryComponent::UCharacterTrajectoryComponent(const FObjectInitializer& ObjectInitializer)
+	: Super(ObjectInitializer)
+{
+	PrimaryComponentTick.bCanEverTick = false;
+	PrimaryComponentTick.bStartWithTickEnabled = false;
+
+	bWantsInitializeComponent = true;
+}
+
+void UCharacterTrajectoryComponent::InitializeComponent()
+{
+	Super::InitializeComponent();
+
+	if (ACharacter* Character = Cast<ACharacter>(GetOwner()))
+	{
+		Character->OnCharacterMovementUpdated.AddDynamic(this, &UCharacterTrajectoryComponent::OnMovementUpdated);
+	}
+	else
+	{
+		UE_LOG(LogMotionTrajectory, Error, TEXT("UCharacterTrajectoryComponent requires its owner to be ACharacter"));
+	}
+
+	NumHistorySamples = FMath::CeilToInt32(HistoryLengthSeconds * HistorySamplesPerSecond);
+	SecondsPerHistorySample = 1.f / HistorySamplesPerSecond;
+
+	const int32 NumPredictionSamples = FMath::CeilToInt32(PredictionLengthSeconds * PredictionSamplesPerSecond);
+	SecondsPerPredictionSample = 1.f / PredictionSamplesPerSecond;
+
+	// History + current sample + prediction
+	Trajectory.Samples.Init(FPoseSearchQueryTrajectorySample(), NumHistorySamples + 1 + NumPredictionSamples);
+}
+
+void UCharacterTrajectoryComponent::UninitializeComponent()
+{
+	if (ACharacter* Character = Cast<ACharacter>(GetOwner()))
+	{
+		Character->OnCharacterMovementUpdated.RemoveDynamic(this, &UCharacterTrajectoryComponent::OnMovementUpdated);
+	}
+	else
+	{
+		UE_LOG(LogMotionTrajectory, Error, TEXT("UCharacterTrajectoryComponent requires its owner to be ACharacter"));
+	}
+
+	Super::UninitializeComponent();
+}
+
+void UCharacterTrajectoryComponent::BeginPlay()
+{
+	Super::BeginPlay();
+
+	const ACharacter* Character = Cast<ACharacter>(GetOwner());
+	if (!ensureMsgf(Character, TEXT("UCharacterTrajectoryComponent requires valid ACharacter owner.")))
+	{
+		return;
+	}
+
+	SkelMeshComponent = Character->GetMesh();
+	if (!ensureMsgf(SkelMeshComponent, TEXT("UCharacterTrajectoryComponent must be run on an ACharacter with a valid USkeletalMeshComponent.")))
+	{
+		return;
+	}
+
+	CharacterMovementComponent = Character->GetCharacterMovement();
+	if (!ensureMsgf(CharacterMovementComponent, TEXT("UCharacterTrajectoryComponent must be run on an ACharacter with a valid UCharacterMovementComponent.")))
+	{
+		return;
+	}
+
+	SkelMeshComponentTransformWS = SkelMeshComponent->GetComponentTransform();
+}
+
+void UCharacterTrajectoryComponent::OnMovementUpdated(float DeltaSeconds, FVector OldLocation, FVector OldVelocity)
+{
+	if (FMath::IsNearlyZero(DeltaSeconds))
+	{
+		return;
+	}
+
+	if (!ensure(CharacterMovementComponent != nullptr && SkelMeshComponent != nullptr))
+	{
+		return;
+	}
+
+	const FTransform PreviousSkelMeshComponentTransformWS = SkelMeshComponentTransformWS;
+	SkelMeshComponentTransformWS = SkelMeshComponent->GetComponentTransform();
+
+	const FTransform SkelMeshTransformDelta = SkelMeshComponentTransformWS.GetRelativeTransform(PreviousSkelMeshComponentTransformWS);
+	UpdateHistory(DeltaSeconds, SkelMeshTransformDelta);
+
+	const FVector VelocityCS = SkelMeshComponentTransformWS.InverseTransformVectorNoScale(CharacterMovementComponent->Velocity);
+	const FVector AccelerationCS = SkelMeshComponentTransformWS.InverseTransformVectorNoScale(CharacterMovementComponent->GetCurrentAcceleration());
+	UpdatePrediction(VelocityCS , AccelerationCS);
+
+#if ENABLE_ANIM_DEBUG
+	if (CVarCharacterTrajectoryDebug.GetValueOnAnyThread() == 1)
+	{
+		Trajectory.DebugDrawTrajectory(GetWorld(), SkelMeshComponent->GetComponentTransform());
+	}
+#endif // ENABLE_ANIM_DEBUG
+
+	// Convert to FTrajectorySampleRange for testing with the Motion Matching node.
+	// This is temporary to support compatibility with existing data until we switch Motion Matching to use FPoseSearchTrajectory.
+	Temp_TrajectorySampleRange.Samples.Reset(Trajectory.Samples.Num());
+	for (const FPoseSearchQueryTrajectorySample& Sample : Trajectory.Samples)
+	{
+		FTrajectorySample TempSample;
+		TempSample.Transform.SetLocation(Sample.Position);
+		TempSample.AccumulatedSeconds = Sample.AccumulatedSeconds;
+
+		Temp_TrajectorySampleRange.Samples.Add(TempSample);
+	}
+}
+
+void UpdateHistorySample(FPoseSearchQueryTrajectorySample& Sample, float DeltaSeconds, const FTransform& DeltaTransformCS)
+{
+	Sample.Position = DeltaTransformCS.InverseTransformPosition(Sample.Position);
+	Sample.AccumulatedSeconds -= DeltaSeconds;
+}
+
+// This function moves each history sample by the inverse of the character's current motion (i.e. if the character is moving forward, the history
+// samples move backward). It also shifts the range of history samples whenever a new history sample should be recorded.
+// This allows us to keep a single sample array in component space that can be read directly by the Motion Matching node, rather than storing world
+// transforms in a separate list and converting them to component space each update. 
+// This also allows us to create "faked" trajectories that match animation data rather than the simulation (e.g. if our animation data only has coverage
+// for one speed, we can adjust the history by a single speed to produce trajectories that best match the data).
+void UCharacterTrajectoryComponent::UpdateHistory(float DeltaSeconds, const FTransform& DeltaTransformCS)
+{
+	check(NumHistorySamples <= Trajectory.Samples.Num());
+
+	// Shift history Samples when it's time to record a new one.
+	if (NumHistorySamples > 0 &&
+		(FMath::Abs(Trajectory.Samples[NumHistorySamples - 1].AccumulatedSeconds) >= SecondsPerHistorySample))
+	{
+		for (int32 Index = 0; Index < NumHistorySamples; ++Index)
+		{
+			Trajectory.Samples[Index] = Trajectory.Samples[Index + 1];
+
+			UpdateHistorySample(Trajectory.Samples[Index], DeltaSeconds, DeltaTransformCS);
+		}
+	}
+	else
+	{
+		for (int32 Index = 0; Index < NumHistorySamples; ++Index)
+		{
+			UpdateHistorySample(Trajectory.Samples[Index], DeltaSeconds, DeltaTransformCS);
+		}
+	}
+}
+
+void UCharacterTrajectoryComponent::UpdatePrediction(const FVector& VelocityCS, const FVector& AccelerationCS)
+{
+	check(CharacterMovementComponent);
+
+	FVector CurrentPositionCS = FVector::ZeroVector;
+	FVector CurrentVelocityCS = VelocityCS;
+	float AccumulatedSeconds = 0.f;
+
+	for (int32 Index = NumHistorySamples + 1; Index < Trajectory.Samples.Num(); ++Index)
+	{
+		CurrentPositionCS += CurrentVelocityCS * SecondsPerPredictionSample;
+		AccumulatedSeconds += SecondsPerPredictionSample;
+
+		Trajectory.Samples[Index].Position = CurrentPositionCS;
+		Trajectory.Samples[Index].AccumulatedSeconds = AccumulatedSeconds;
+
+		FVector NewVelocityCS = FVector::ZeroVector;
+		UCharacterMovementTrajectoryLibrary::StepCharacterMovementGroundPrediction(SecondsPerPredictionSample, CurrentVelocityCS, AccelerationCS, CharacterMovementComponent, NewVelocityCS);
+		CurrentVelocityCS = NewVelocityCS;
+	}
+}
\ No newline at end of file
diff --git a/Engine/Plugins/Experimental/Animation/MotionTrajectory/Source/MotionTrajectory/Public/CharacterMovementTrajectoryLibrary.h b/Engine/Plugins/Experimental/Animation/MotionTrajectory/Source/MotionTrajectory/Public/CharacterMovementTrajectoryLibrary.h
new file mode 100644
--- /dev/null
+++ b/Engine/Plugins/Experimental/Animation/MotionTrajectory/Source/MotionTrajectory/Public/CharacterMovementTrajectoryLibrary.h
@@ -0,0 +1,33 @@
+// Copyright Epic Games, Inc. All Rights Reserved.
+
+#pragma once
+
+#include "Kismet/BlueprintFunctionLibrary.h"
+
+#include "CharacterMovementTrajectoryLibrary.generated.h"
+
+class UCharacterMovementComponent;
+
+/**
+ * Library of functions useful for generating trajectories based on the behavior of the UCharacterMovementComponent.
+ */
+UCLASS()
+class MOTIONTRAJECTORY_API UCharacterMovementTrajectoryLibrary : public UBlueprintFunctionLibrary
+{
+	GENERATED_BODY()
+
+public:
+
+	/**
+	 * Approximate the behavior of the UCharacterMovementComponent during ground locomotion.
+	 * @param InDeltaTime - Length of time to predict forward. Running multiple shorter iterations will generate more accurate results than one large time step.
+	 * @param InVelocity - Current velocity.
+	 * @param InAcceleration - Current acceleration.
+	 * @param InCharacterMovementComponent - Tuning values, such as GroundFriction, are read from the movement component.
+	 * @param OutVelocity - The resulting predicted velocity.
+	 */
+	UFUNCTION(BlueprintPure, Category = "Motion Trajectory")
+	static void StepCharacterMovementGroundPrediction(
+		float InDeltaTime, const FVector& InVelocity, const FVector& InAcceleration, const UCharacterMovementComponent* InCharacterMovementComponent,
+		FVector& OutVelocity);
+};
\ No newline at end of file
diff --git a/Engine/Plugins/Experimental/Animation/MotionTrajectory/Source/MotionTrajectory/Public/CharacterTrajectoryComponent.h b/Engine/Plugins/Experimental/Animation/MotionTrajectory/Source/MotionTrajectory/Public/CharacterTrajectoryComponent.h
new file mode 100644
--- /dev/null
+++ b/Engine/Plugins/Experimental/Animation/MotionTrajectory/Source/MotionTrajectory/Public/CharacterTrajectoryComponent.h
@@ -0,0 +1,81 @@
+// Copyright Epic Games, Inc. All Rights Reserved.
+
+#pragma once
+
+#include "Animation/MotionTrajectoryTypes.h"
+#include "Components/ActorComponent.h"
+#include "PoseSearch/PoseSearchTrajectoryTypes.h"
+
+// This is temporary to support compatibility with existing data until we switch Motion Matching to use FPoseSearchTrajectory.
+#include "Animation/MotionTrajectoryTypes.h"
+
+#include "CharacterTrajectoryComponent.generated.h"
+
+class UCharacterMovementComponent;
+
+// Component for generating trajectories usable by Motion Matching. This component generates trajectories from ACharacter.
+// This is intended to provide an example and starting point for using Motion Matching with a common setup using the default UCharacterMovementComponent.
+// It is expected work flow to extend or replace this component for projects that use a custom movement component or custom movement modes.
+UCLASS(BlueprintType, meta = (BlueprintSpawnableComponent), Experimental)
+class MOTIONTRAJECTORY_API UCharacterTrajectoryComponent : public UActorComponent
+{
+	GENERATED_BODY()
+
+public:
+	
+	UCharacterTrajectoryComponent(const FObjectInitializer& ObjectInitializer);
+
+	// Begin UActorComponent Interface
+	virtual void InitializeComponent() override;
+	virtual void UninitializeComponent() override;
+	virtual void BeginPlay() override;
+	// End UActorComponent Interface
+
+protected:
+	UFUNCTION()
+	void OnMovementUpdated(float DeltaSeconds, FVector OldLocation, FVector OldVelocity);
+
+	void UpdateHistory(float DeltaSeconds, const FTransform& DeltaTransformCS);
+	void UpdatePrediction(const FVector& VelocityCS, const FVector& AccelerationCS);
+
+protected:
+
+	// Trajectory stored in component space so it can be directly passed to Motion Matching.
+	UPROPERTY(BlueprintReadOnly, Category = "Trajectory")
+	FPoseSearchQueryTrajectory Trajectory;
+
+	// This should generally match the longest history required by a Motion Matching Database in the project.
+	// Motion Matching will use extrapolation to generate samples if the history doesn't contain enough samples.
+	UPROPERTY(EditDefaultsOnly, BlueprintReadOnly, Category = "Trajectory Settings")
+	float HistoryLengthSeconds = 1.f;
+
+	// Higher values will cost more storage and processing time, but give higher accuracy.
+	UPROPERTY(EditDefaultsOnly, BlueprintReadOnly, Category = "Trajectory Settings", meta = (ClampMin = "1", ClampMax = "120"))
+	int32 HistorySamplesPerSecond = 15;
+
+	// This should match the longest trajectory prediction required by a Motion Matching Database in the project.
+	// Motion Matching will use extrapolation to generate samples if the prediction doesn't contain enough samples.
+	UPROPERTY(EditDefaultsOnly, BlueprintReadOnly, Category = "Trajectory Settings")
+	float PredictionLengthSeconds = 1.f;
+
+	// Higher values will cost more storage and processing time, but give higher accuracy.
+	UPROPERTY(EditDefaultsOnly, BlueprintReadOnly, Category = "Trajectory Settings", meta = (ClampMin = "1", ClampMax = "120"))
+	int32 PredictionSamplesPerSecond = 15;
+
+	UPROPERTY()
+	TObjectPtr<USkeletalMeshComponent> SkelMeshComponent;
+
+	UPROPERTY()
+	TObjectPtr<UCharacterMovementComponent> CharacterMovementComponent;
+
+	int32 NumHistorySamples = -1;
+	float SecondsPerHistorySample = 0.f;
+	float SecondsPerPredictionSample = 0.f;
+
+	// Current transform of the skeletal mesh component, used to calculate the movement delta between frames.
+	FTransform SkelMeshComponentTransformWS = FTransform::Identity;
+
+	// This is temporary to support compatibility with existing data until we switch Motion Matching to use FPoseSearchQueryTrajectory.
+	UPROPERTY(BlueprintReadOnly, Category = "Trajectory")
+	FTrajectorySampleRange Temp_TrajectorySampleRange;
+};
\ No newline at end of file
diff --git a/Engine/Plugins/Experimental/Animation/PoseSearch/Source/Runtime/Private/PoseSearchTrajectoryTypes.cpp b/Engine/Plugins/Experimental/Animation/PoseSearch/Source/Runtime/Private/PoseSearchTrajectoryTypes.cpp
new file mode 100644
--- /dev/null
+++ b/Engine/Plugins/Experimental/Animation/PoseSearch/Source/Runtime/Private/PoseSearchTrajectoryTypes.cpp
@@ -0,0 +1,86 @@
+// Copyright Epic Games, Inc. All Rights Reserved.
+
+#include "PoseSearch/PoseSearchTrajectoryTypes.h"
+
+#include "Animation/AnimTypes.h"
+#include "DrawDebugHelpers.h"
+
+FPoseSearchQueryTrajectorySample FPoseSearchQueryTrajectorySample::Lerp(const FPoseSearchQueryTrajectorySample& Other, float Alpha) const
+{
+	check(Facing.IsNormalized());
+	check(Other.Facing.IsNormalized());
+
+	FPoseSearchQueryTrajectorySample Result;
+	
+	Result.Facing = FQuat::FastLerp(Facing, Other.Facing, Alpha);
+	Result.Position = FMath::Lerp(Position, Other.Position, Alpha);
+	Result.AccumulatedSeconds = FMath::Lerp(AccumulatedSeconds, Other.AccumulatedSeconds, Alpha);
+	
+	Result.Facing.Normalize();
+
+	return Result;
+}
+
+FPoseSearchQueryTrajectorySample FPoseSearchQueryTrajectory::GetSampleAtTime(float Time, bool bExtrapolate) const
+{
+	const int32 Num = Samples.Num();
+	if (Num > 1)
+	{
+		const int32 LowerBoundIdx = Algo::LowerBound(Samples, Time, [](const FPoseSearchQueryTrajectorySample& TrajectorySample, float Value)
+			{
+				return Value > TrajectorySample.AccumulatedSeconds;
+			});
+
+		const int32 NextIdx = FMath::Clamp(LowerBoundIdx, 1, Samples.Num() - 1);
+		const int32 PrevIdx = NextIdx - 1;
+
+		const float Denominator = Samples[NextIdx].AccumulatedSeconds - Samples[PrevIdx].AccumulatedSeconds;
+		if (!FMath::IsNearlyZero(Denominator))
+		{
+			const float Numerator = Time - Samples[PrevIdx].AccumulatedSeconds;
+			const float LerpValue = bExtrapolate ? Numerator / Denominator : FMath::Clamp(Numerator / Denominator, 0.f, 1.f);
+			return Samples[PrevIdx].Lerp(Samples[NextIdx], LerpValue);
+		}
+
+		return Samples[PrevIdx];
+	}
+
+	if (Num > 0)
+	{
+		return Samples[0];
+	}
+
+	return FPoseSearchQueryTrajectorySample();
+}
+
+#if ENABLE_ANIM_DEBUG
+void FPoseSearchQueryTrajectory::DebugDrawTrajectory(const UWorld* World, const FTransform& TransformWS) const
+{
+	if (Samples.Num() <= 1)
+	{
+		return;
+	}
+
+	for (int32 Index = 0; Index < Samples.Num(); ++Index)
+	{
+		const FVector CurrentSamplePositionWS = TransformWS.TransformPosition(Samples[Index].Position);
+
+		DrawDebugSphere(
+			World,
+			CurrentSamplePositionWS,
+			2.f /*Radius*/, 4 /*Segments*/,
+			FColor::Black, false /*bPersistentLines*/, -1.f /*LifeTime*/, 0 /*DepthPriority*/, 1.f /*Thickness*/);
+
+		if (Samples.IsValidIndex(Index + 1))
+		{
+			const FVector NextSamplePositionWS = TransformWS.TransformPosition(Samples[Index + 1].Position);
+
+			DrawDebugLine(
+				World,
+				CurrentSamplePositionWS,
+				NextSamplePositionWS,
+				FColor::Black, false /*bPersistentLines*/, -1.f /*LifeTime*/, 0 /*DepthPriority*/, 1.f /*Thickness*/);
+		}
+	}
+}
+#endif // ENABLE_ANIM_DEBUG
\ No newline at end of file
diff --git a/Engine/Plugins/Experimental/Animation/PoseSearch/Source/Runtime/Public/PoseSearch/AnimNode_MotionMatching.h b/Engine/Plugins/Experimental/Animation/PoseSearch/Source/Runtime/Public/PoseSearch/AnimNode_MotionMatching.h
--- a/Engine/Plugins/Experimental/Animation/PoseSearch/Source/Runtime/Public/PoseSearch/AnimNode_MotionMatching.h
+++ b/Engine/Plugins/Experimental/Animation/PoseSearch/Source/Runtime/Public/PoseSearch/AnimNode_MotionMatching.h
@@ -27,7 +27,7 @@ public:
 	UPROPERTY(EditAnywhere, BlueprintReadWrite, Category = Settings, meta = (PinShownByDefault))
 	FGameplayTagContainer ActiveTagsContainer;
 
-	// Motion trajectory samples for pose search queries
+	// Motion trajectory samples for pose search queries. Expected to be in the space of the SkeletalMeshComponent.
 	UPROPERTY(EditAnywhere, BlueprintReadWrite, Category=Settings, meta=(PinShownByDefault))
 	FTrajectorySampleRange Trajectory;
 
diff --git a/Engine/Plugins/Experimental/Animation/PoseSearch/Source/Runtime/Public/PoseSearch/PoseSearchTrajectoryTypes.h b/Engine/Plugins/Experimental/Animation/PoseSearch/Source/Runtime/Public/PoseSearch/PoseSearchTrajectoryTypes.h
new file mode 100644
--- /dev/null
+++ b/Engine/Plugins/Experimental/Animation/PoseSearch/Source/Runtime/Public/PoseSearch/PoseSearchTrajectoryTypes.h
@@ -0,0 +1,40 @@
+// Copyright Epic Games, Inc. All Rights Reserved.
+
+#pragma once
+
+#include "Animation/AnimTypes.h"
+
+#include "PoseSearchTrajectoryTypes.generated.h"
+
+USTRUCT(BlueprintType, Category="Pose Search Trajectory")
+struct POSESEARCH_API FPoseSearchQueryTrajectorySample
+{
+	GENERATED_BODY()
+
+	UPROPERTY(EditAnywhere, BlueprintReadWrite, Category = "Pose Search Query Trajectory")
+	FQuat Facing = FQuat::Identity;
+
+	UPROPERTY(EditAnywhere, BlueprintReadWrite, Category = "Pose Search Query Trajectory")
+	FVector Position = FVector::ZeroVector;
+
+	UPROPERTY(EditAnywhere, BlueprintReadWrite, Category = "Pose Search Query Trajectory")
+	float AccumulatedSeconds = 0.f;
+
+	FPoseSearchQueryTrajectorySample Lerp(const FPoseSearchQueryTrajectorySample& Other, float Alpha) const;
+};
+
+USTRUCT(BlueprintType, Category = "Motion Trajectory")
+struct POSESEARCH_API FPoseSearchQueryTrajectory
+{
+	GENERATED_BODY()
+
+	// This contains zero or more history samples, a current sample, and zero or more future predicted samples.
+	UPROPERTY(EditAnywhere, BlueprintReadWrite, Category = "Pose Search Query Trajectory")
+	TArray<FPoseSearchQueryTrajectorySample> Samples;
+
+	FPoseSearchQueryTrajectorySample GetSampleAtTime(float Time, bool bExtrapolate) const;
+
+#if ENABLE_ANIM_DEBUG
+	void DebugDrawTrajectory(const UWorld* World, const FTransform& TransformWS) const;
+#endif // ENABLE_ANIM_DEBUG
+};
\ No newline at end of file
