Motion Matching: Support rotation in Character Trajectory

diff --git a/Engine/Plugins/Experimental/Animation/MotionTrajectory/Source/MotionTrajectory/Private/CharacterTrajectoryComponent.cpp b/Engine/Plugins/Experimental/Animation/MotionTrajectory/Source/MotionTrajectory/Private/CharacterTrajectoryComponent.cpp
--- a/Engine/Plugins/Experimental/Animation/MotionTrajectory/Source/MotionTrajectory/Private/CharacterTrajectoryComponent.cpp
+++ b/Engine/Plugins/Experimental/Animation/MotionTrajectory/Source/MotionTrajectory/Private/CharacterTrajectoryComponent.cpp
@@ -37,15 +37,6 @@ void UCharacterTrajectoryComponent::InitializeComponent()
 	{
 		UE_LOG(LogMotionTrajectory, Error, TEXT("UCharacterTrajectoryComponent requires its owner to be ACharacter"));
 	}
-
-	NumHistorySamples = FMath::CeilToInt32(HistoryLengthSeconds * HistorySamplesPerSecond);
-	SecondsPerHistorySample = 1.f / HistorySamplesPerSecond;
-
-	const int32 NumPredictionSamples = FMath::CeilToInt32(PredictionLengthSeconds * PredictionSamplesPerSecond);
-	SecondsPerPredictionSample = 1.f / PredictionSamplesPerSecond;
-
-	// History + current sample + prediction
-	Trajectory.Samples.Init(FPoseSearchQueryTrajectorySample(), NumHistorySamples + 1 + NumPredictionSamples);
 }
 
 void UCharacterTrajectoryComponent::UninitializeComponent()
@@ -85,6 +76,28 @@ void UCharacterTrajectoryComponent::BeginPlay()
 	}
 
 	SkelMeshComponentTransformWS = SkelMeshComponent->GetComponentTransform();
+
+	// Default forward in the engine is the X axis, but data often diverges from this (e.g. it's common for skeletal meshes to be Y forward).
+	// We determine the forward direction in the space of the skeletal mesh component based on the offset from the actor.
+	ForwardFacingCS = SkelMeshComponent->GetRelativeRotation().Quaternion().Inverse();
+
+	// The UI clamps these to be non-zero.
+	check(HistorySamplesPerSecond);
+	check(PredictionSamplesPerSecond);
+
+	NumHistorySamples = FMath::CeilToInt32(HistoryLengthSeconds * HistorySamplesPerSecond);
+	SecondsPerHistorySample = 1.f / HistorySamplesPerSecond;
+
+	const int32 NumPredictionSamples = FMath::CeilToInt32(PredictionLengthSeconds * PredictionSamplesPerSecond);
+	SecondsPerPredictionSample = 1.f / PredictionSamplesPerSecond;
+
+	FPoseSearchQueryTrajectorySample DefaultSample;
+	DefaultSample.Facing = ForwardFacingCS;
+	DefaultSample.Position = FVector::ZeroVector;
+	DefaultSample.AccumulatedSeconds = 0.f;
+
+	// History + current sample + prediction
+	Trajectory.Samples.Init(DefaultSample, NumHistorySamples + 1 + NumPredictionSamples);
 }
 
 void UCharacterTrajectoryComponent::OnMovementUpdated(float DeltaSeconds, FVector OldLocation, FVector OldVelocity)
@@ -107,12 +120,13 @@ void UCharacterTrajectoryComponent::OnMovementUpdated(float DeltaSeconds, FVecto
 
 	const FVector VelocityCS = SkelMeshComponentTransformWS.InverseTransformVectorNoScale(CharacterMovementComponent->Velocity);
 	const FVector AccelerationCS = SkelMeshComponentTransformWS.InverseTransformVectorNoScale(CharacterMovementComponent->GetCurrentAcceleration());
-	UpdatePrediction(VelocityCS , AccelerationCS);
+	const FRotator ControllerRotationRate = CalculateControllerRotationRate(DeltaSeconds, CharacterMovementComponent->ShouldRemainVertical());
+	UpdatePrediction(VelocityCS, AccelerationCS, ControllerRotationRate);
 
 #if ENABLE_ANIM_DEBUG
 	if (CVarCharacterTrajectoryDebug.GetValueOnAnyThread() == 1)
 	{
-		Trajectory.DebugDrawTrajectory(GetWorld(), SkelMeshComponent->GetComponentTransform());
+		Trajectory.DebugDrawTrajectory(GetWorld(), SkelMeshComponentTransformWS);
 	}
 #endif // ENABLE_ANIM_DEBUG
 
@@ -131,6 +145,7 @@ void UCharacterTrajectoryComponent::OnMovementUpdated(float DeltaSeconds, FVecto
 
 void UpdateHistorySample(FPoseSearchQueryTrajectorySample& Sample, float DeltaSeconds, const FTransform& DeltaTransformCS)
 {
+	Sample.Facing = DeltaTransformCS.InverseTransformRotation(Sample.Facing);
 	Sample.Position = DeltaTransformCS.InverseTransformPosition(Sample.Position);
 	Sample.AccumulatedSeconds -= DeltaSeconds;
 }
@@ -165,24 +180,71 @@ void UCharacterTrajectoryComponent::UpdateHistory(float DeltaSeconds, const FTra
 	}
 }
 
-void UCharacterTrajectoryComponent::UpdatePrediction(const FVector& VelocityCS, const FVector& AccelerationCS)
+void UCharacterTrajectoryComponent::UpdatePrediction(const FVector& VelocityCS, const FVector& AccelerationCS, const FRotator& ControllerRotationRate)
 {
 	check(CharacterMovementComponent);
 
 	FVector CurrentPositionCS = FVector::ZeroVector;
 	FVector CurrentVelocityCS = VelocityCS;
+	FVector CurrentAccelerationCS = AccelerationCS;
+	FQuat CurrentFacingCS = ForwardFacingCS;
 	float AccumulatedSeconds = 0.f;
 
+	FQuat ControllerRotationPerStep = (ControllerRotationRate * SecondsPerPredictionSample).Quaternion();
+
 	for (int32 Index = NumHistorySamples + 1; Index < Trajectory.Samples.Num(); ++Index)
 	{
 		CurrentPositionCS += CurrentVelocityCS * SecondsPerPredictionSample;
 		AccumulatedSeconds += SecondsPerPredictionSample;
 
+		// Account for the controller (e.g. the camera) rotating.
+		CurrentFacingCS = ControllerRotationPerStep * CurrentFacingCS;
+		CurrentAccelerationCS = ControllerRotationPerStep * CurrentAccelerationCS;
+
 		Trajectory.Samples[Index].Position = CurrentPositionCS;
+		Trajectory.Samples[Index].Facing = CurrentFacingCS;
 		Trajectory.Samples[Index].AccumulatedSeconds = AccumulatedSeconds;
 
 		FVector NewVelocityCS = FVector::ZeroVector;
-		UCharacterMovementTrajectoryLibrary::StepCharacterMovementGroundPrediction(SecondsPerPredictionSample, CurrentVelocityCS, AccelerationCS, CharacterMovementComponent, NewVelocityCS);
+		UCharacterMovementTrajectoryLibrary::StepCharacterMovementGroundPrediction(
+			SecondsPerPredictionSample, CurrentVelocityCS, CurrentAccelerationCS, CharacterMovementComponent,
+			NewVelocityCS);
 		CurrentVelocityCS = NewVelocityCS;
+
+		if (CharacterMovementComponent->bOrientRotationToMovement && !CurrentAccelerationCS.IsNearlyZero())
+		{
+			// Rotate towards acceleration.
+			CurrentFacingCS = FMath::QInterpConstantTo(CurrentFacingCS, CurrentAccelerationCS.ToOrientationQuat(), SecondsPerPredictionSample, RotateTowardsMovementSpeed);
+		}
+	}
+}
+
+// Calculate how much the character is rotating each update due to the controller (e.g. the camera) rotating.
+// E.g. If the user is moving forward but rotating the camera, the character (and thus future accelerations, facing directions, etc) will rotate.
+FRotator UCharacterTrajectoryComponent::CalculateControllerRotationRate(float DeltaSeconds, bool bShouldRemainVertical)
+{
+	check(GetOwner());
+
+	// OnMovementUpdated handles DeltaSeconds == 0.f, so we should never hit this.
+	check(DeltaSeconds > 0.f);
+
+	ACharacter* CharacterOwner = Cast<ACharacter>(GetOwner());
+	if (CharacterOwner == nullptr || CharacterOwner->Controller == nullptr)
+	{
+		// @todo: Simulated proxies don't have controllers, so they'll need some other mechanism to account for controller rotation rate.
+		return FRotator::ZeroRotator;
 	}
+
+	FRotator DesiredControllerRotation = CharacterOwner->Controller->GetDesiredRotation();
+	if (bShouldRemainVertical)
+	{
+		DesiredControllerRotation.Yaw = FRotator::NormalizeAxis(DesiredControllerRotation.Yaw);
+		DesiredControllerRotation.Pitch = 0.f;
+		DesiredControllerRotation.Roll = 0.f;
+	}
+
+	const FRotator DesiredRotationDelta = DesiredControllerRotation - DesiredControllerRotationLastUpdate;
+	DesiredControllerRotationLastUpdate = DesiredControllerRotation;
+
+	return DesiredRotationDelta.GetNormalized() * (1.f / DeltaSeconds);
 }
\ No newline at end of file
diff --git a/Engine/Plugins/Experimental/Animation/MotionTrajectory/Source/MotionTrajectory/Public/CharacterTrajectoryComponent.h b/Engine/Plugins/Experimental/Animation/MotionTrajectory/Source/MotionTrajectory/Public/CharacterTrajectoryComponent.h
--- a/Engine/Plugins/Experimental/Animation/MotionTrajectory/Source/MotionTrajectory/Public/CharacterTrajectoryComponent.h
+++ b/Engine/Plugins/Experimental/Animation/MotionTrajectory/Source/MotionTrajectory/Public/CharacterTrajectoryComponent.h
@@ -36,7 +36,9 @@ protected:
 	void OnMovementUpdated(float DeltaSeconds, FVector OldLocation, FVector OldVelocity);
 
 	void UpdateHistory(float DeltaSeconds, const FTransform& DeltaTransformCS);
-	void UpdatePrediction(const FVector& VelocityCS, const FVector& AccelerationCS);
+	void UpdatePrediction(const FVector& VelocityCS, const FVector& AccelerationCS, const FRotator& ControllerRotationRate);
+
+	FRotator CalculateControllerRotationRate(float DeltaSeconds, bool bShouldRemainVertical);
 
 protected:
 
@@ -47,20 +49,26 @@ protected:
 	// This should generally match the longest history required by a Motion Matching Database in the project.
 	// Motion Matching will use extrapolation to generate samples if the history doesn't contain enough samples.
 	UPROPERTY(EditDefaultsOnly, BlueprintReadOnly, Category = "Trajectory Settings")
-	float HistoryLengthSeconds = 1.f;
+	float HistoryLengthSeconds = 1.5f;
 
 	// Higher values will cost more storage and processing time, but give higher accuracy.
 	UPROPERTY(EditDefaultsOnly, BlueprintReadOnly, Category = "Trajectory Settings", meta = (ClampMin = "1", ClampMax = "120"))
-	int32 HistorySamplesPerSecond = 15;
+	int32 HistorySamplesPerSecond = 5;
 
 	// This should match the longest trajectory prediction required by a Motion Matching Database in the project.
 	// Motion Matching will use extrapolation to generate samples if the prediction doesn't contain enough samples.
 	UPROPERTY(EditDefaultsOnly, BlueprintReadOnly, Category = "Trajectory Settings")
-	float PredictionLengthSeconds = 1.f;
+	float PredictionLengthSeconds = 1.5f;
 
 	// Higher values will cost more storage and processing time, but give higher accuracy.
 	UPROPERTY(EditDefaultsOnly, BlueprintReadOnly, Category = "Trajectory Settings", meta = (ClampMin = "1", ClampMax = "120"))
-	int32 PredictionSamplesPerSecond = 15;
+	int32 PredictionSamplesPerSecond = 5;
+
+	// If the character is forward facing (i.e. bOrientRotationToMovement is true), this controls how quickly the trajectory will rotate
+	// to face acceleration. It's common for this to differ from the rotation rate of the character, because animations are often authored 
+	// with different rotation speeds than the character. This is especially true in cases where the character rotation snaps to movement.
+	UPROPERTY(EditDefaultsOnly, BlueprintReadOnly, Category = "Trajectory Settings")
+	float RotateTowardsMovementSpeed = 10.f;
 
 	UPROPERTY()
 	TObjectPtr<USkeletalMeshComponent> SkelMeshComponent;
@@ -75,6 +83,11 @@ protected:
 	// Current transform of the skeletal mesh component, used to calculate the movement delta between frames.
 	FTransform SkelMeshComponentTransformWS = FTransform::Identity;
 
+	// Forward axis for the SkeletalMeshComponent. It's common for skeletal mesh and animation data to not be X forward.
+	FQuat ForwardFacingCS = FQuat::Identity;
+
+	FRotator DesiredControllerRotationLastUpdate = FRotator::ZeroRotator;
+
 	// This is temporary to support compatibility with existing data until we switch Motion Matching to use FPoseSearchQueryTrajectory.
 	UPROPERTY(BlueprintReadOnly, Category = "Trajectory")
 	FTrajectorySampleRange Temp_TrajectorySampleRange;
diff --git a/Engine/Plugins/Experimental/Animation/PoseSearch/Source/Runtime/Private/PoseSearchTrajectoryTypes.cpp b/Engine/Plugins/Experimental/Animation/PoseSearch/Source/Runtime/Private/PoseSearchTrajectoryTypes.cpp
--- a/Engine/Plugins/Experimental/Animation/PoseSearch/Source/Runtime/Private/PoseSearchTrajectoryTypes.cpp
+++ b/Engine/Plugins/Experimental/Animation/PoseSearch/Source/Runtime/Private/PoseSearchTrajectoryTypes.cpp
@@ -81,6 +81,13 @@ void FPoseSearchQueryTrajectory::DebugDrawTrajectory(const UWorld* World, const
 				NextSamplePositionWS,
 				FColor::Black, false /*bPersistentLines*/, -1.f /*LifeTime*/, 0 /*DepthPriority*/, 1.f /*Thickness*/);
 		}
+
+		const FQuat CurrentSampleFacingWS = TransformWS.TransformRotation(Samples[Index].Facing);
+		DrawDebugDirectionalArrow(
+			World,
+			CurrentSamplePositionWS,
+			CurrentSamplePositionWS + CurrentSampleFacingWS.RotateVector(FVector::ForwardVector) * 25.f,
+			20.f, FColor::Orange, false /*bPersistentLines*/, -1.f /*LifeTime*/, 0 /*DepthPriority*/, 1.f /*Thickness*/);
 	}
 }
 #endif // ENABLE_ANIM_DEBUG
\ No newline at end of file
