motion matching
- replacing FTrajectorySampleRange with FPoseSearchQueryTrajectory in PoseSearch plugin
- deprecating FTrajectorySample and FTrajectorySampleRange

diff --git a/Engine/Plugins/Experimental/Animation/MotionTrajectory/Source/MotionTrajectory/Private/CharacterTrajectoryComponent.cpp b/Engine/Plugins/Experimental/Animation/MotionTrajectory/Source/MotionTrajectory/Private/CharacterTrajectoryComponent.cpp
--- a/Engine/Plugins/Experimental/Animation/MotionTrajectory/Source/MotionTrajectory/Private/CharacterTrajectoryComponent.cpp
+++ b/Engine/Plugins/Experimental/Animation/MotionTrajectory/Source/MotionTrajectory/Private/CharacterTrajectoryComponent.cpp
@@ -9,9 +9,6 @@
 #include "HAL/IConsoleManager.h"
 #include "MotionTrajectory.h"
 
-#include "MotionTrajectoryCharacterMovement.h"
-#include "MotionTrajectoryLibrary.h"
-
 #if ENABLE_ANIM_DEBUG
 TAutoConsoleVariable<int32> CVarCharacterTrajectoryDebug(TEXT("a.CharacterTrajectory.Debug"), 0, TEXT("Turn on debug drawing for Character trajectory"));
 #endif // ENABLE_ANIM_DEBUG
@@ -98,15 +95,22 @@ void UCharacterTrajectoryComponent::BeginPlay()
 
 	// History + current sample + prediction
 	Trajectory.Samples.Init(DefaultSample, NumHistorySamples + 1 + NumPredictionSamples);
-}
 
-void UCharacterTrajectoryComponent::OnMovementUpdated(float DeltaSeconds, FVector OldLocation, FVector OldVelocity)
-{
-	if (FMath::IsNearlyZero(DeltaSeconds))
+	// initializing history samples AccumulatedSeconds
+	for (int32 i = 0; i < NumHistorySamples; ++i)
 	{
-		return;
+		Trajectory.Samples[i].AccumulatedSeconds = SecondsPerHistorySample * (i - NumHistorySamples);
+	}
+
+	// initializing history samples AccumulatedSeconds
+	for (int32 i = NumHistorySamples + 1; i < Trajectory.Samples.Num(); ++i)
+	{
+		Trajectory.Samples[i].AccumulatedSeconds = SecondsPerPredictionSample * (i - NumHistorySamples);
 	}
+}
 
+void UCharacterTrajectoryComponent::OnMovementUpdated(float DeltaSeconds, FVector OldLocation, FVector OldVelocity)
+{
 	if (!ensure(CharacterMovementComponent != nullptr && SkelMeshComponent != nullptr))
 	{
 		return;
@@ -124,23 +128,21 @@ void UCharacterTrajectoryComponent::OnMovementUpdated(float DeltaSeconds, FVecto
 	UpdatePrediction(VelocityCS, AccelerationCS, ControllerRotationRate);
 
 #if ENABLE_ANIM_DEBUG
-	if (CVarCharacterTrajectoryDebug.GetValueOnAnyThread() == 1)
+	if (CVarCharacterTrajectoryDebug.GetValueOnAnyThread())
 	{
 		Trajectory.DebugDrawTrajectory(GetWorld(), SkelMeshComponentTransformWS);
 	}
 #endif // ENABLE_ANIM_DEBUG
+}
 
-	// Convert to FTrajectorySampleRange for testing with the Motion Matching node.
-	// This is temporary to support compatibility with existing data until we switch Motion Matching to use FPoseSearchTrajectory.
-	Temp_TrajectorySampleRange.Samples.Reset(Trajectory.Samples.Num());
-	for (const FPoseSearchQueryTrajectorySample& Sample : Trajectory.Samples)
-	{
-		FTrajectorySample TempSample;
-		TempSample.Transform.SetLocation(Sample.Position);
-		TempSample.AccumulatedSeconds = Sample.AccumulatedSeconds;
-
-		Temp_TrajectorySampleRange.Samples.Add(TempSample);
-	}
+FPoseSearchQueryTrajectory UCharacterTrajectoryComponent::GetCharacterRelativeTrajectory() const
+{
+	FPoseSearchQueryTrajectory CharacterRelativeTrajectory = Trajectory;
+	const FTransform OwnerTransformWS = GetOwner()->GetActorTransform();
+	const FTransform ComponentTransformWS = SkelMeshComponent->GetComponentTransform();
+	const FTransform ReferenceChangeTransform = OwnerTransformWS.GetRelativeTransform(ComponentTransformWS);
+	CharacterRelativeTrajectory.TransformReferenceFrame(ReferenceChangeTransform);
+	return CharacterRelativeTrajectory;
 }
 
 void UpdateHistorySample(FPoseSearchQueryTrajectorySample& Sample, float DeltaSeconds, const FTransform& DeltaTransformCS)
@@ -161,8 +163,7 @@ void UCharacterTrajectoryComponent::UpdateHistory(float DeltaSeconds, const FTra
 	check(NumHistorySamples <= Trajectory.Samples.Num());
 
 	// Shift history Samples when it's time to record a new one.
-	if (NumHistorySamples > 0 &&
-		(FMath::Abs(Trajectory.Samples[NumHistorySamples - 1].AccumulatedSeconds) >= SecondsPerHistorySample))
+	if (NumHistorySamples > 0 && FMath::Abs(Trajectory.Samples[NumHistorySamples - 1].AccumulatedSeconds) >= SecondsPerHistorySample)
 	{
 		for (int32 Index = 0; Index < NumHistorySamples; ++Index)
 		{
diff --git a/Engine/Plugins/Experimental/Animation/MotionTrajectory/Source/MotionTrajectory/Private/MotionTrajectory.cpp b/Engine/Plugins/Experimental/Animation/MotionTrajectory/Source/MotionTrajectory/Private/MotionTrajectory.cpp
--- a/Engine/Plugins/Experimental/Animation/MotionTrajectory/Source/MotionTrajectory/Private/MotionTrajectory.cpp
+++ b/Engine/Plugins/Experimental/Animation/MotionTrajectory/Source/MotionTrajectory/Private/MotionTrajectory.cpp
@@ -7,8 +7,6 @@
 #include "Kismet/KismetSystemLibrary.h"
 #include "Modules/ModuleManager.h"
 
-#include UE_INLINE_GENERATED_CPP_BY_NAME(MotionTrajectory)
-
 #define LOCTEXT_NAMESPACE "MotionTrajectory"
 
 DEFINE_LOG_CATEGORY(LogMotionTrajectory);
@@ -21,216 +19,4 @@ void FMotionTrajectoryModule::ShutdownModule()
 {
 }
 
-UMotionTrajectoryComponent::UMotionTrajectoryComponent(const FObjectInitializer& ObjectInitializer) : Super(ObjectInitializer)
-{
-	PrimaryComponentTick.bCanEverTick = true;
-	PrimaryComponentTick.bStartWithTickEnabled = true;
-	PrimaryComponentTick.TickGroup = TG_PrePhysics;
-}
-
-FTrajectorySample UMotionTrajectoryComponent::CalcWorldSpacePresentTrajectorySample(float DeltaTime) const
-{
-	const APawn* Pawn = TryGetOwnerPawn();
-	UE_LOG(LogMotionTrajectory, Error, TEXT("UMotionTrajectoryComponent::GetPresentTrajectory for Pawn: %s requires implementation."), Pawn ? *(Pawn->GetHumanReadableName()) : TEXT("NULL"));
-	return FTrajectorySample();
-}
-
-void UMotionTrajectoryComponent::TickTrajectory(float DeltaTime)
-{
-	// Compute the instantaneous/present trajectory sample and guarantee that is zeroed on all accumulated domains
-	PresentTrajectorySampleWS = CalcWorldSpacePresentTrajectorySample(DeltaTime);
-	PresentTrajectorySampleWS.AccumulatedSeconds = 0.f;
-
-	// convert to local space
-	PresentTrajectorySampleLS = PresentTrajectorySampleWS;
-	PresentTrajectorySampleLS.PrependOffset(PresentTrajectorySampleWS.Transform.Inverse(), 0.0f);
-
-
-	// Tick the historical sample retention/decay algorithm by one iteration
-	TickHistoryEvictionPolicy();
-}
-
-const APawn* UMotionTrajectoryComponent::TryGetOwnerPawn() const
-{
-	const AActor* Actor = GetOwner();
-	return Actor ? Cast<APawn>(Actor) : nullptr;
-}
-
-FTrajectorySampleRange UMotionTrajectoryComponent::CombineHistoryPresentPrediction(bool bIncludeHistory, const FTrajectorySampleRange& Prediction) const
-{
-	FTrajectorySampleRange Trajectory;
-	int32 TotalSampleCount = 1 + Prediction.Samples.Num(); // present + future
-
-	FTrajectorySampleRange History;
-	if (bIncludeHistory)
-	{
-		History = GetHistory();
-		TotalSampleCount += History.Samples.Num();
-	}
-
-	// Linearly combines all uniformly sampled trajectory ranges from [0...n (samples)] in domain ordering:
-	// Ex: [Past + Present + Future], with negative values representing the past, a zeroed present, and positive values representing the future
-	Trajectory.Samples.Reserve(TotalSampleCount);
-	Trajectory.Samples.Append(History.Samples);
-	Trajectory.Samples.Add(PresentTrajectorySampleLS);
-	Trajectory.Samples.Append(Prediction.Samples);
-
-#if WITH_EDITORONLY_DATA
-	Trajectory.DebugDrawTrajectory(bDebugDrawTrajectory, GetOwner()->GetWorld(), PreviousWorldTransform);
-#endif
-	return Trajectory;
-}
-
-void UMotionTrajectoryComponent::TickHistoryEvictionPolicy()
-{
-	// World space transform and time are used for determining local-relative historical samples
-	const UWorld* World = GetWorld();
-	const FTransform WorldTransform = PresentTrajectorySampleWS.Transform;
-	const float WorldGameTime = UKismetSystemLibrary::GetGameTimeInSeconds(World);
-
-	// Skip on the very first sample
-	if (PreviousWorldGameTime != 0.f)
-	{
-		// Compute world space/time deltas for historical sample decay
-		const float DeltaSeconds = WorldGameTime - PreviousWorldGameTime;
-		const float DeltaDistance = FVector::Distance(WorldTransform.GetLocation(), PreviousWorldTransform.GetLocation());
-		const FTransform DeltaTransform = PreviousWorldTransform.GetRelativeTransform(WorldTransform);
-
-		for (auto& Sample : SampleHistory)
-		{
-			Sample.PrependOffset(DeltaTransform, -DeltaSeconds);
-		}
-
-		const FTrajectorySample FirstSample = SampleHistory.First();
-		const float FirstSampleTime = FirstSample.AccumulatedSeconds;
-
-		// For both time and distance history domains, we compute the concept of an "effective" time horizon:
-		// This value is present to guarantee uniform history decay against a fixed anchor point in the past when there is lack of trajectory motion
-		// The anchor point is defined by the AccumulatedTime of the furthest history sample at the exact moment of zero motion
-		if (FMath::IsNearlyZero(DeltaDistance) && !FirstSample.IsZeroSample())
-		{
-			if (EffectiveTimeDomain == 0.f)
-			{
-				// One time initialization of an effective time domain at the exact moment of zero motion
-				EffectiveTimeDomain = FirstSampleTime;
-			}
-		}
-		else
-		{
-			// When we are in motion, an effective time domain is not required
-			EffectiveTimeDomain = 0.f;
-		}
-
-		// Remove all trajectory samples which are outside the boundaries of the enabled domain horizons
-		SampleHistory.RemoveAll([&](const FTrajectorySample& Sample)
-			{
-				// Remove superfluous zero motion samples
-				if (Sample.IsZeroSample() && PresentTrajectorySampleLS.IsZeroSample())
-				{
-					return true;
-				}
-
-				// Time horizon
-				if (Sample.AccumulatedSeconds < -HistorySettings.Seconds)
-				{
-					return true;
-				}
-
-				// Effective time horizon
-				if (EffectiveTimeDomain != 0.f&& (Sample.AccumulatedSeconds < EffectiveTimeDomain))
-				{
-					return true;
-				}
-
-				return false;
-			});
-	}
-
-	SampleHistory.Emplace(PresentTrajectorySampleLS);
-	
-	// Cache the present world space sample as the next frame's reference point for delta computation
-	PreviousWorldTransform = WorldTransform;
-	PreviousWorldGameTime = WorldGameTime;
-}
-
-void UMotionTrajectoryComponent::FlushHistory()
-{
-	SampleHistory.Reset();
-	PreviousWorldTransform = {};
-	PreviousWorldGameTime = 0.f;
-	EffectiveTimeDomain = 0.f;
-}
-
-void UMotionTrajectoryComponent::OnComponentCreated()
-{
-	FlushHistory();
-	SampleHistory.Reserve(MaxSamples);
-	Super::OnComponentCreated();
-}
-
-void UMotionTrajectoryComponent::BeginPlay() 
-{
-	Super::BeginPlay();
-
-	FTransform PresentTransform = FTransform::Identity;
-
-	if (const APawn* Pawn = TryGetOwnerPawn())
-	{
-		PresentTransform = Pawn->GetActorTransform();
-	}
-
-	PresentTrajectorySampleWS = FTrajectorySample();
-	PresentTrajectorySampleWS.Transform = PresentTransform;
-
-	PresentTrajectorySampleLS = FTrajectorySample();
-}
-
-void UMotionTrajectoryComponent::TickComponent(float DeltaTime, enum ELevelTick TickType, FActorComponentTickFunction* ThisTickFunction)
-{
-	Super::TickComponent(DeltaTime, TickType, ThisTickFunction);
-
-	TickTrajectory(DeltaTime);
-}
-
-FTrajectorySampleRange UMotionTrajectoryComponent::GetTrajectory() const
-{
-	const APawn* Pawn = TryGetOwnerPawn();
-	checkf(false, TEXT("UMotionTrajectoryComponent::GetTrajectory for Pawn: %s requires implementation."), Pawn ? *(Pawn->GetHumanReadableName()) : TEXT("NULL"));
-	return FTrajectorySampleRange();
-}
-
-FTrajectorySampleRange UMotionTrajectoryComponent::GetTrajectoryWithSettings(const FMotionTrajectorySettings& Settings, bool bIncludeHistory) const
-{
-	const APawn* Pawn = TryGetOwnerPawn();
-	checkf(false, TEXT("UMotionTrajectoryComponent::GetTrajectoryWithSettings for Pawn: %s requires implementation."), Pawn ? *(Pawn->GetHumanReadableName()) : TEXT("NULL"));
-	return FTrajectorySampleRange();
-}
-
-FTrajectorySampleRange UMotionTrajectoryComponent::GetHistory() const
-{
-	FTrajectorySampleRange SourceHistory;
-	const int32 HistorySampleCount = SampleHistory.Num();
-	if (HistorySampleCount > 0)
-	{
-		const int32 PresentSampleIdx = HistorySampleCount - 1;
-
-		// Convert the source history ring buffer into a trajectory range
-		// We consider these values to be non-uniform with a variable, frame-rate dependent sample rate governed by the component tick rate
-		SourceHistory.Samples.Reserve(HistorySampleCount);
-
-		for (const auto& Sample : SampleHistory)
-		{
-			SourceHistory.Samples.Emplace(Sample);
-		}
-
-		// @todo: fix me - maybe don't add it in the first place?
-		// Remove the present trajectory sample, this will not be included in the history
-		check(FMath::IsNearlyZero(SourceHistory.Samples[PresentSampleIdx].AccumulatedSeconds));
-		SourceHistory.Samples.RemoveAt(PresentSampleIdx);
-	}
-	return SourceHistory;
-}
-
-#undef LOCTEXT_NAMESPACE
-
 IMPLEMENT_MODULE(FMotionTrajectoryModule, MotionTrajectory)
diff --git a/Engine/Plugins/Experimental/Animation/MotionTrajectory/Source/MotionTrajectory/Private/MotionTrajectoryCharacterMovement.cpp b/Engine/Plugins/Experimental/Animation/MotionTrajectory/Source/MotionTrajectory/Private/MotionTrajectoryCharacterMovement.cpp
deleted file mode 100644
diff --git a/Engine/Plugins/Experimental/Animation/MotionTrajectory/Source/MotionTrajectory/Private/MotionTrajectoryLibrary.cpp b/Engine/Plugins/Experimental/Animation/MotionTrajectory/Source/MotionTrajectory/Private/MotionTrajectoryLibrary.cpp
deleted file mode 100644
diff --git a/Engine/Plugins/Experimental/Animation/MotionTrajectory/Source/MotionTrajectory/Public/CharacterTrajectoryComponent.h b/Engine/Plugins/Experimental/Animation/MotionTrajectory/Source/MotionTrajectory/Public/CharacterTrajectoryComponent.h
--- a/Engine/Plugins/Experimental/Animation/MotionTrajectory/Source/MotionTrajectory/Public/CharacterTrajectoryComponent.h
+++ b/Engine/Plugins/Experimental/Animation/MotionTrajectory/Source/MotionTrajectory/Public/CharacterTrajectoryComponent.h
@@ -2,13 +2,8 @@
 
 #pragma once
 
-#include "Animation/MotionTrajectoryTypes.h"
 #include "Components/ActorComponent.h"
 #include "PoseSearch/PoseSearchTrajectoryTypes.h"
-
-// This is temporary to support compatibility with existing data until we switch Motion Matching to use FPoseSearchTrajectory.
-#include "Animation/MotionTrajectoryTypes.h"
-
 #include "CharacterTrajectoryComponent.generated.h"
 
 class UCharacterMovementComponent;
@@ -35,13 +30,14 @@ protected:
 	UFUNCTION()
 	void OnMovementUpdated(float DeltaSeconds, FVector OldLocation, FVector OldVelocity);
 
+	UFUNCTION(BlueprintCallable, Category="Trajectory")
+	FPoseSearchQueryTrajectory GetCharacterRelativeTrajectory() const;
+
 	void UpdateHistory(float DeltaSeconds, const FTransform& DeltaTransformCS);
 	void UpdatePrediction(const FVector& VelocityCS, const FVector& AccelerationCS, const FRotator& ControllerRotationRate);
 
 	FRotator CalculateControllerRotationRate(float DeltaSeconds, bool bShouldRemainVertical);
 
-protected:
-
 	// Trajectory stored in component space so it can be directly passed to Motion Matching.
 	UPROPERTY(BlueprintReadOnly, Category = "Trajectory")
 	FPoseSearchQueryTrajectory Trajectory;
@@ -87,8 +83,4 @@ protected:
 	FQuat ForwardFacingCS = FQuat::Identity;
 
 	FRotator DesiredControllerRotationLastUpdate = FRotator::ZeroRotator;
-
-	// This is temporary to support compatibility with existing data until we switch Motion Matching to use FPoseSearchQueryTrajectory.
-	UPROPERTY(BlueprintReadOnly, Category = "Trajectory")
-	FTrajectorySampleRange Temp_TrajectorySampleRange;
 };
\ No newline at end of file
diff --git a/Engine/Plugins/Experimental/Animation/MotionTrajectory/Source/MotionTrajectory/Public/MotionTrajectory.h b/Engine/Plugins/Experimental/Animation/MotionTrajectory/Source/MotionTrajectory/Public/MotionTrajectory.h
--- a/Engine/Plugins/Experimental/Animation/MotionTrajectory/Source/MotionTrajectory/Public/MotionTrajectory.h
+++ b/Engine/Plugins/Experimental/Animation/MotionTrajectory/Source/MotionTrajectory/Public/MotionTrajectory.h
@@ -2,7 +2,6 @@
 
 #pragma once
 
-#include "Animation/MotionTrajectoryTypes.h"
 #include "Components/ActorComponent.h"
 #include "Containers/RingBuffer.h"
 
@@ -32,109 +31,6 @@ struct MOTIONTRAJECTORY_API FMotionTrajectorySettings
 	float Seconds = 2.f;
 };
 
-// Abstract component interface for the minimum Motion Trajectory prediction and history API
-UCLASS(ClassGroup=Movement, abstract, BlueprintType, Category="Motion Trajectory", Experimental)
-class MOTIONTRAJECTORY_API UMotionTrajectoryComponent : public UActorComponent
-{
-	GENERATED_BODY()
-
-private:
-
-	// Previous frame world-space transform
-	FTransform PreviousWorldTransform = {};
-
-	// Previous frame world game time
-	float PreviousWorldGameTime = 0.f;
-
-	// A superseding "effective" time domain applied history samples, which guarantees uniform decay
-	float EffectiveTimeDomain = 0.f;
-
-	// Component-driven tick which decays and evicts trajectory history outside of the defined domain horizons
-	void TickHistoryEvictionPolicy();
-
-protected:
-
-	// Retained trajectory history samples
-	TRingBuffer<FTrajectorySample> SampleHistory = {};
-
-	// Retained instantaneous/present trajectory in world space
-	FTrajectorySample PresentTrajectorySampleWS = {};
-
-	// Retained instantaneous/present trajectory in actor space
-	FTrajectorySample PresentTrajectorySampleLS = {};
-
-	// Gets the instantaneous/present trajectory sample of the current frame
-	virtual FTrajectorySample CalcWorldSpacePresentTrajectorySample(float DeltaTime) const;
-
-	// Ticks the trajectory. Call this to tick before ::TickComponent when necessary
-	virtual void TickTrajectory(float DeltaTime);
-
-	// Gets the Pawn from the owning Actor of this component
-	const APawn* TryGetOwnerPawn() const;
-
-	// Combines all trajectory samples in the past, present, and future into a unified trajectory range
-	FTrajectorySampleRange CombineHistoryPresentPrediction(bool bIncludeHistory, const FTrajectorySampleRange& Prediction) const;
-
-
-	// Forcefully evicts all trajectory history and resets internal history tracking state
-	void FlushHistory();
-
-public:
-
-	UMotionTrajectoryComponent(const FObjectInitializer& ObjectInitializer);
-
-	// Begin UActorComponent Interface
-	virtual void OnComponentCreated() override;
-	virtual void BeginPlay() override;
-	virtual void TickComponent(float DeltaTime
-		, enum ELevelTick TickType
-		, FActorComponentTickFunction* ThisTickFunction) override;
-	// End UActorComponent Interface
-
-	// Interface for computing a trajectory prediction
-	UFUNCTION(BlueprintPure, Category="Motion Trajectory")
-	virtual FTrajectorySampleRange GetTrajectory() const;
-
-	// Interface for computing a trajectory prediction with overridden component settings
-	UFUNCTION(BlueprintPure, Category="Motion Trajectory", meta=(AutoCreateRefTerm="Settings", DisplayName="Get Trajectory (With Settings)"))
-	virtual FTrajectorySampleRange GetTrajectoryWithSettings(const FMotionTrajectorySettings& Settings, bool bIncludeHistory) const;
-	
-	// Retrieves the historical trajectory
-	UFUNCTION(BlueprintCallable, Category="Motion Trajectory")
-	FTrajectorySampleRange GetHistory() const;
-
-	// Prediction trajectory simulation settings
-	UPROPERTY(EditDefaultsOnly, BlueprintReadOnly, Category="Settings")
-	FMotionTrajectorySettings PredictionSettings;
-
-	// Historical trajectory settings
-	UPROPERTY(EditDefaultsOnly, BlueprintReadOnly, Category="Settings")
-	FMotionTrajectorySettings HistorySettings;
-
-	// The trajectory sampling rate for prediction
-	UPROPERTY(EditDefaultsOnly, BlueprintReadOnly, Category="Settings", meta=(ClampMin="5", ClampMax="120"))
-	int32 SampleRate = 30;
-
-	// The maximum number of trajectory samples retained by the internal history buffer
-	UPROPERTY(EditDefaultsOnly, BlueprintReadOnly, Category="Settings", meta=(ClampMin="0", ClampMax="3600"))
-	int32 MaxSamples = 1000;
-
-	// Determines if GetTrajectory() will automatically combine history into the returned trajectory prediction (true)
-	// Consider disabling this option when historical trajectory samples are never used
-	// Example: No Motion Matching historical sample times or distances are defined in the Pose Search Schema
-	UPROPERTY(EditDefaultsOnly, BlueprintReadOnly, Category="Settings")
-	bool bPredictionIncludesHistory = true;
-
-#if WITH_EDITORONLY_DATA
-	// Debug drawing options for the complete trajectory range
-	// a.MotionTrajectory.Debug (0/1) to enable or disable
-	// a.MotionTrajectory.Options (0-6) to enable additional sample information
-	// a.MotionTrajectory.Stride (*) to stride sample information display by a specified modulo
-	UPROPERTY(EditAnywhere, Category="Debug")
-	bool bDebugDrawTrajectory = false;
-#endif
-};
-
 #if UE_ENABLE_INCLUDE_ORDER_DEPRECATED_IN_5_2
 #include "Modules/ModuleManager.h"
 #endif
diff --git a/Engine/Plugins/Experimental/Animation/MotionTrajectory/Source/MotionTrajectory/Public/MotionTrajectoryCharacterMovement.h b/Engine/Plugins/Experimental/Animation/MotionTrajectory/Source/MotionTrajectory/Public/MotionTrajectoryCharacterMovement.h
deleted file mode 100644
diff --git a/Engine/Plugins/Experimental/Animation/MotionTrajectory/Source/MotionTrajectory/Public/MotionTrajectoryLibrary.h b/Engine/Plugins/Experimental/Animation/MotionTrajectory/Source/MotionTrajectory/Public/MotionTrajectoryLibrary.h
deleted file mode 100644
diff --git a/Engine/Plugins/Experimental/Animation/PoseSearch/Source/Runtime/Private/AnimNode_MotionMatching.cpp b/Engine/Plugins/Experimental/Animation/PoseSearch/Source/Runtime/Private/AnimNode_MotionMatching.cpp
--- a/Engine/Plugins/Experimental/Animation/PoseSearch/Source/Runtime/Private/AnimNode_MotionMatching.cpp
+++ b/Engine/Plugins/Experimental/Animation/PoseSearch/Source/Runtime/Private/AnimNode_MotionMatching.cpp
@@ -7,7 +7,6 @@
 #include "Animation/AnimRootMotionProvider.h"
 #include "Animation/AnimSequence.h"
 #include "Animation/BlendSpace.h"
-#include "Animation/MotionTrajectoryTypes.h"
 #include "Components/SkeletalMeshComponent.h"
 #include "PoseSearch/AnimNode_PoseSearchHistoryCollector.h"
 #include "PoseSearch/PoseSearchDatabase.h"
diff --git a/Engine/Plugins/Experimental/Animation/PoseSearch/Source/Runtime/Private/PoseSearchContext.cpp b/Engine/Plugins/Experimental/Animation/PoseSearch/Source/Runtime/Private/PoseSearchContext.cpp
--- a/Engine/Plugins/Experimental/Animation/PoseSearch/Source/Runtime/Private/PoseSearchContext.cpp
+++ b/Engine/Plugins/Experimental/Animation/PoseSearch/Source/Runtime/Private/PoseSearchContext.cpp
@@ -2,11 +2,11 @@
 
 #include "PoseSearch/PoseSearchContext.h"
 #include "Animation/AnimInstanceProxy.h"
-#include "Animation/MotionTrajectoryTypes.h"
 #include "AnimationRuntime.h"
 #include "PoseSearch/PoseSearchDatabase.h"
 #include "PoseSearch/PoseSearchHistory.h"
 #include "PoseSearch/PoseSearchSchema.h"
+#include "PoseSearch/PoseSearchTrajectoryTypes.h"
 #include "PoseSearchFeatureChannel_Position.h"
 
 namespace UE::PoseSearch
@@ -259,7 +259,7 @@ void FDebugDrawParams::DrawFeatureVector(int32 PoseIdx)
 
 //////////////////////////////////////////////////////////////////////////
 // FSearchContext
-FSearchContext::FSearchContext(const FTrajectorySampleRange* InTrajectory, const IPoseHistory* InHistory, float InDesiredPermutationTimeOffset, const FPoseIndicesHistory* InPoseIndicesHistory,
+FSearchContext::FSearchContext(const FPoseSearchQueryTrajectory* InTrajectory, const IPoseHistory* InHistory, float InDesiredPermutationTimeOffset, const FPoseIndicesHistory* InPoseIndicesHistory,
 	EPoseSearchBooleanRequest InQueryMirrorRequest, const FSearchResult& InCurrentResult, float InPoseJumpThresholdTime, bool bInForceInterrupt, bool bInCanAdvance)
 : Trajectory(InTrajectory)
 , History(InHistory)
@@ -362,8 +362,8 @@ FTransform FSearchContext::GetTransform(float SampleTime, const UPoseSearchSchem
 	else
 	{
 		check(Trajectory);
-		const FTrajectorySample TrajectorySample = Trajectory->GetSampleAtTime(SampleTime);
-		RootTransform = TrajectorySample.Transform;
+		const FPoseSearchQueryTrajectorySample TrajectorySample = Trajectory->GetSampleAtTime(SampleTime);
+		RootTransform = TrajectorySample.GetTransform();
 	}
 
 	const FBoneIndexType BoneIndexType = Schema->GetBoneIndexType(SchemaBoneIdx);
diff --git a/Engine/Plugins/Experimental/Animation/PoseSearch/Source/Runtime/Private/PoseSearchFeatureChannel_Heading.cpp b/Engine/Plugins/Experimental/Animation/PoseSearch/Source/Runtime/Private/PoseSearchFeatureChannel_Heading.cpp
--- a/Engine/Plugins/Experimental/Animation/PoseSearch/Source/Runtime/Private/PoseSearchFeatureChannel_Heading.cpp
+++ b/Engine/Plugins/Experimental/Animation/PoseSearch/Source/Runtime/Private/PoseSearchFeatureChannel_Heading.cpp
@@ -1,7 +1,6 @@
 // Copyright Epic Games, Inc. All Rights Reserved.
 
 #include "PoseSearchFeatureChannel_Heading.h"
-#include "Animation/MotionTrajectoryTypes.h"
 #include "PoseSearch/PoseSearchAssetIndexer.h"
 #include "PoseSearch/PoseSearchAssetSampler.h"
 #include "PoseSearch/PoseSearchContext.h"
diff --git a/Engine/Plugins/Experimental/Animation/PoseSearch/Source/Runtime/Private/PoseSearchFeatureChannel_Position.cpp b/Engine/Plugins/Experimental/Animation/PoseSearch/Source/Runtime/Private/PoseSearchFeatureChannel_Position.cpp
--- a/Engine/Plugins/Experimental/Animation/PoseSearch/Source/Runtime/Private/PoseSearchFeatureChannel_Position.cpp
+++ b/Engine/Plugins/Experimental/Animation/PoseSearch/Source/Runtime/Private/PoseSearchFeatureChannel_Position.cpp
@@ -1,7 +1,6 @@
 // Copyright Epic Games, Inc. All Rights Reserved.
 
 #include "PoseSearchFeatureChannel_Position.h"
-#include "Animation/MotionTrajectoryTypes.h"
 #include "PoseSearch/PoseSearchAssetIndexer.h"
 #include "PoseSearch/PoseSearchAssetSampler.h"
 #include "PoseSearch/PoseSearchContext.h"
diff --git a/Engine/Plugins/Experimental/Animation/PoseSearch/Source/Runtime/Private/PoseSearchFeatureChannel_Trajectory.h b/Engine/Plugins/Experimental/Animation/PoseSearch/Source/Runtime/Private/PoseSearchFeatureChannel_Trajectory.h
--- a/Engine/Plugins/Experimental/Animation/PoseSearch/Source/Runtime/Private/PoseSearchFeatureChannel_Trajectory.h
+++ b/Engine/Plugins/Experimental/Animation/PoseSearch/Source/Runtime/Private/PoseSearchFeatureChannel_Trajectory.h
@@ -6,8 +6,6 @@
 #include "UObject/ObjectSaveContext.h"
 #include "PoseSearchFeatureChannel_Trajectory.generated.h"
 
-struct FTrajectorySampleRange;
-
 UENUM(meta = (Bitflags, UseEnumValuesAsMaskValuesInEditor = "true"))
 enum class EPoseSearchTrajectoryFlags : uint32
 {
diff --git a/Engine/Plugins/Experimental/Animation/PoseSearch/Source/Runtime/Private/PoseSearchFeatureChannel_Velocity.cpp b/Engine/Plugins/Experimental/Animation/PoseSearch/Source/Runtime/Private/PoseSearchFeatureChannel_Velocity.cpp
--- a/Engine/Plugins/Experimental/Animation/PoseSearch/Source/Runtime/Private/PoseSearchFeatureChannel_Velocity.cpp
+++ b/Engine/Plugins/Experimental/Animation/PoseSearch/Source/Runtime/Private/PoseSearchFeatureChannel_Velocity.cpp
@@ -1,7 +1,6 @@
 // Copyright Epic Games, Inc. All Rights Reserved.
 
 #include "PoseSearchFeatureChannel_Velocity.h"
-#include "Animation/MotionTrajectoryTypes.h"
 #include "PoseSearch/PoseSearchAssetIndexer.h"
 #include "PoseSearch/PoseSearchAssetSampler.h"
 #include "PoseSearch/PoseSearchContext.h"
diff --git a/Engine/Plugins/Experimental/Animation/PoseSearch/Source/Runtime/Private/PoseSearchLibrary.cpp b/Engine/Plugins/Experimental/Animation/PoseSearch/Source/Runtime/Private/PoseSearchLibrary.cpp
--- a/Engine/Plugins/Experimental/Animation/PoseSearch/Source/Runtime/Private/PoseSearchLibrary.cpp
+++ b/Engine/Plugins/Experimental/Animation/PoseSearch/Source/Runtime/Private/PoseSearchLibrary.cpp
@@ -228,10 +228,10 @@ void UPoseSearchLibrary::TraceMotionMatchingState(
 	if (DeltaTime > SMALL_NUMBER && SearchContext.GetTrajectory())
 	{
 		// simulation
-		const FTrajectorySample PrevSample = SearchContext.GetTrajectory()->GetSampleAtTime(-DeltaTime);
-		const FTrajectorySample CurrSample = SearchContext.GetTrajectory()->GetSampleAtTime(0.f);
+		const FPoseSearchQueryTrajectorySample PrevSample = SearchContext.GetTrajectory()->GetSampleAtTime(-DeltaTime);
+		const FPoseSearchQueryTrajectorySample CurrSample = SearchContext.GetTrajectory()->GetSampleAtTime(0.f);
 
-		const FTransform SimDelta = CurrSample.Transform.GetRelativeTransform(PrevSample.Transform);
+		const FTransform SimDelta = CurrSample.GetTransform().GetRelativeTransform(PrevSample.GetTransform());
 
 		TraceState.SimLinearVelocity = SimDelta.GetTranslation().Size() / DeltaTime;
 		TraceState.SimAngularVelocity = FMath::RadiansToDegrees(SimDelta.GetRotation().GetAngle()) / DeltaTime;
@@ -253,7 +253,7 @@ void UPoseSearchLibrary::TraceMotionMatchingState(
 void UPoseSearchLibrary::UpdateMotionMatchingState(
 	const FAnimationUpdateContext& Context,
 	const TArray<TObjectPtr<const UPoseSearchDatabase>>& Databases,
-	const FTrajectorySampleRange& Trajectory,
+	const FPoseSearchQueryTrajectory& Trajectory,
 	const FMotionMatchingSettings& Settings,
 	FMotionMatchingState& InOutMotionMatchingState,
 	bool bForceInterrupt)
@@ -401,7 +401,7 @@ void UPoseSearchLibrary::UpdateMotionMatchingState(
 void UPoseSearchLibrary::MotionMatch(
 	UAnimInstance* AnimInstance,
 	const UPoseSearchDatabase* Database,
-	const FTrajectorySampleRange Trajectory,
+	const FPoseSearchQueryTrajectory Trajectory,
 	const FName PoseHistoryName,
 	UAnimationAsset*& SelectedAnimation,
 	float& SelectedTime,
@@ -508,9 +508,9 @@ void UPoseSearchLibrary::MotionMatch(
 					FCSPose<FCompactPose> ComponentSpacePose;
 					ComponentSpacePose.InitPose(Pose);
 
-					const FTrajectorySample TrajectorySample = Trajectory.GetSampleAtTime(ExtractionTime);
+					const FPoseSearchQueryTrajectorySample TrajectorySample = Trajectory.GetSampleAtTime(ExtractionTime);
 					const FTransform& ComponentTransform = AnimInstance->GetOwningComponent()->GetComponentTransform();
-					const FTransform FutureComponentTransform = TrajectorySample.Transform * ComponentTransform;
+					const FTransform FutureComponentTransform = TrajectorySample.GetTransform() * ComponentTransform;
 
 					ExtendedPoseHistory.AddFuturePose(FutureAnimationTime, ComponentSpacePose, FutureComponentTransform);
 				}
diff --git a/Engine/Plugins/Experimental/Animation/PoseSearch/Source/Runtime/Private/PoseSearchTrajectoryTypes.cpp b/Engine/Plugins/Experimental/Animation/PoseSearch/Source/Runtime/Private/PoseSearchTrajectoryTypes.cpp
--- a/Engine/Plugins/Experimental/Animation/PoseSearch/Source/Runtime/Private/PoseSearchTrajectoryTypes.cpp
+++ b/Engine/Plugins/Experimental/Animation/PoseSearch/Source/Runtime/Private/PoseSearchTrajectoryTypes.cpp
@@ -3,7 +3,9 @@
 #include "PoseSearch/PoseSearchTrajectoryTypes.h"
 
 #include "Animation/AnimTypes.h"
+#include "Components/SceneComponent.h"
 #include "DrawDebugHelpers.h"
+#include "PoseSearch/PoseSearchDefines.h"
 
 FPoseSearchQueryTrajectorySample FPoseSearchQueryTrajectorySample::Lerp(const FPoseSearchQueryTrajectorySample& Other, float Alpha) const
 {
@@ -12,15 +14,19 @@ FPoseSearchQueryTrajectorySample FPoseSearchQueryTrajectorySample::Lerp(const FP
 
 	FPoseSearchQueryTrajectorySample Result;
 	
-	Result.Facing = FQuat::FastLerp(Facing, Other.Facing, Alpha);
+	Result.Facing = FQuat::FastLerp(Facing, Other.Facing, Alpha).GetNormalized();
 	Result.Position = FMath::Lerp(Position, Other.Position, Alpha);
 	Result.AccumulatedSeconds = FMath::Lerp(AccumulatedSeconds, Other.AccumulatedSeconds, Alpha);
-	
-	Result.Facing.Normalize();
 
 	return Result;
 }
 
+void FPoseSearchQueryTrajectorySample::SetTransform(const FTransform& Transform)
+{
+	Position = Transform.GetTranslation();
+	Facing = Transform.GetRotation();
+}
+
 FPoseSearchQueryTrajectorySample FPoseSearchQueryTrajectory::GetSampleAtTime(float Time, bool bExtrapolate) const
 {
 	const int32 Num = Samples.Num();
@@ -53,14 +59,19 @@ FPoseSearchQueryTrajectorySample FPoseSearchQueryTrajectory::GetSampleAtTime(flo
 	return FPoseSearchQueryTrajectorySample();
 }
 
-#if ENABLE_ANIM_DEBUG
-void FPoseSearchQueryTrajectory::DebugDrawTrajectory(const UWorld* World, const FTransform& TransformWS) const
+void FPoseSearchQueryTrajectory::TransformReferenceFrame(const FTransform& DeltaTransform)
 {
-	if (Samples.Num() <= 1)
+	const FTransform InverseDeltaTransform = DeltaTransform.Inverse();
+	for (FPoseSearchQueryTrajectorySample& Sample : Samples)
 	{
-		return;
+		const FTransform Transform = InverseDeltaTransform * Sample.GetTransform() * DeltaTransform;
+		Sample.SetTransform(Transform);
 	}
+}
 
+#if ENABLE_ANIM_DEBUG
+void FPoseSearchQueryTrajectory::DebugDrawTrajectory(const UWorld* World, const FTransform& TransformWS) const
+{
 	for (int32 Index = 0; Index < Samples.Num(); ++Index)
 	{
 		const FVector CurrentSamplePositionWS = TransformWS.TransformPosition(Samples[Index].Position);
diff --git a/Engine/Plugins/Experimental/Animation/PoseSearch/Source/Runtime/Public/PoseSearch/AnimNode_MotionMatching.h b/Engine/Plugins/Experimental/Animation/PoseSearch/Source/Runtime/Public/PoseSearch/AnimNode_MotionMatching.h
--- a/Engine/Plugins/Experimental/Animation/PoseSearch/Source/Runtime/Public/PoseSearch/AnimNode_MotionMatching.h
+++ b/Engine/Plugins/Experimental/Animation/PoseSearch/Source/Runtime/Public/PoseSearch/AnimNode_MotionMatching.h
@@ -3,10 +3,10 @@
 #pragma once
 
 #include "Animation/AnimNode_AssetPlayerBase.h"
-#include "Animation/MotionTrajectoryTypes.h"
 #include "GameplayTagContainer.h"
 #include "PoseSearch/AnimNode_BlendStack.h"
 #include "PoseSearch/PoseSearchLibrary.h"
+#include "PoseSearch/PoseSearchTrajectoryTypes.h"
 #include "AnimNode_MotionMatching.generated.h"
 
 class UPoseSearchDatabase;
@@ -26,7 +26,7 @@ private:
 
 	// Motion Trajectory samples for pose search queries in Motion Matching.These are expected to be in the space of the SkeletalMeshComponent.This is provided with the CharacterMovementTrajectory Component output.
 	UPROPERTY(EditAnywhere, Category=Settings, meta=(PinShownByDefault))
-	FTrajectorySampleRange Trajectory;
+	FPoseSearchQueryTrajectory Trajectory;
 
 	// Settings for the core motion matching node.
 	UPROPERTY(EditAnywhere, Category=Settings, meta=(PinHiddenByDefault))
diff --git a/Engine/Plugins/Experimental/Animation/PoseSearch/Source/Runtime/Public/PoseSearch/PoseSearchContext.h b/Engine/Plugins/Experimental/Animation/PoseSearch/Source/Runtime/Public/PoseSearch/PoseSearchContext.h
--- a/Engine/Plugins/Experimental/Animation/PoseSearch/Source/Runtime/Public/PoseSearch/PoseSearchContext.h
+++ b/Engine/Plugins/Experimental/Animation/PoseSearch/Source/Runtime/Public/PoseSearch/PoseSearchContext.h
@@ -8,7 +8,7 @@
 #include "PoseSearch/PoseSearchIndex.h"
 #include "PoseSearch/PoseSearchResult.h"
 
-struct FTrajectorySampleRange;
+struct FPoseSearchQueryTrajectory;
 class UPoseSearchDatabase;
 class UPoseSearchFeatureChannel_Position;
 
@@ -138,7 +138,7 @@ private:
 
 struct POSESEARCH_API FSearchContext
 {
-	FSearchContext(const FTrajectorySampleRange* InTrajectory, const IPoseHistory* InHistory, float InDesiredPermutationTimeOffset,
+	FSearchContext(const FPoseSearchQueryTrajectory* InTrajectory, const IPoseHistory* InHistory, float InDesiredPermutationTimeOffset,
 		const FPoseIndicesHistory* InPoseIndicesHistory = nullptr, EPoseSearchBooleanRequest InQueryMirrorRequest = EPoseSearchBooleanRequest::Indifferent,
 		const FSearchResult& InCurrentResult = FSearchResult(), float InPoseJumpThresholdTime = 0.f, bool bInForceInterrupt = false, bool bInCanAdvance = true);
 
@@ -167,7 +167,7 @@ struct POSESEARCH_API FSearchContext
 	const FPoseIndicesHistory* GetPoseIndicesHistory() const { return PoseIndicesHistory; }
 	const IPoseHistory* GetHistory() const { return History; }
 	float GetDesiredPermutationTimeOffset() const { return DesiredPermutationTimeOffset; }
-	const FTrajectorySampleRange* GetTrajectory() const { return Trajectory; }
+	const FPoseSearchQueryTrajectory* GetTrajectory() const { return Trajectory; }
 	bool CanAdvance() const { return bCanAdvance; }
 	bool IsForceInterrupt() const { return bForceInterrupt; }
 
@@ -176,7 +176,7 @@ private:
 	FTransform GetComponentSpaceTransform(float SampleTime, const UPoseSearchSchema* Schema, int8 SchemaSampleBoneIdx = RootSchemaBoneIdx);
 	FVector GetSamplePositionInternal(float SampleTime, float OriginTime, const UPoseSearchSchema* Schema, int8 SchemaSampleBoneIdx = RootSchemaBoneIdx, int8 SchemaOriginBoneIdx = RootSchemaBoneIdx, bool bUseHistoryRoot = false);
 
-	const FTrajectorySampleRange* Trajectory = nullptr;
+	const FPoseSearchQueryTrajectory* Trajectory = nullptr;
 	const IPoseHistory* History = nullptr;
 	float DesiredPermutationTimeOffset = 0.f;
 	const FPoseIndicesHistory* PoseIndicesHistory = nullptr;
diff --git a/Engine/Plugins/Experimental/Animation/PoseSearch/Source/Runtime/Public/PoseSearch/PoseSearchLibrary.h b/Engine/Plugins/Experimental/Animation/PoseSearch/Source/Runtime/Public/PoseSearch/PoseSearchLibrary.h
--- a/Engine/Plugins/Experimental/Animation/PoseSearch/Source/Runtime/Public/PoseSearch/PoseSearchLibrary.h
+++ b/Engine/Plugins/Experimental/Animation/PoseSearch/Source/Runtime/Public/PoseSearch/PoseSearchLibrary.h
@@ -20,7 +20,7 @@ namespace UE::PoseSearch
 } // namespace UE::PoseSearch
 
 struct FAnimationUpdateContext;
-struct FTrajectorySampleRange;
+struct FPoseSearchQueryTrajectory;
 
 USTRUCT(BlueprintType, Category = "Animation|Pose Search")
 struct POSESEARCH_API FMotionMatchingSettings
@@ -140,7 +140,7 @@ public:
 	static void UpdateMotionMatchingState(
 		const FAnimationUpdateContext& Context,
 		const TArray<TObjectPtr<const UPoseSearchDatabase>>& Databases,
-		const FTrajectorySampleRange& Trajectory,
+		const FPoseSearchQueryTrajectory& Trajectory,
 		const FMotionMatchingSettings& Settings,
 		FMotionMatchingState& InOutMotionMatchingState,
 		bool bForceInterrupt);
@@ -167,7 +167,7 @@ public:
 	static void MotionMatch(
 		UAnimInstance* AnimInstance,
 		const UPoseSearchDatabase* Database,
-		const FTrajectorySampleRange Trajectory,
+		const FPoseSearchQueryTrajectory Trajectory,
 		const FName PoseHistoryName,
 		UAnimationAsset*& SelectedAnimation,
 		float& SelectedTime,
diff --git a/Engine/Plugins/Experimental/Animation/PoseSearch/Source/Runtime/Public/PoseSearch/PoseSearchTrajectoryTypes.h b/Engine/Plugins/Experimental/Animation/PoseSearch/Source/Runtime/Public/PoseSearch/PoseSearchTrajectoryTypes.h
--- a/Engine/Plugins/Experimental/Animation/PoseSearch/Source/Runtime/Public/PoseSearch/PoseSearchTrajectoryTypes.h
+++ b/Engine/Plugins/Experimental/Animation/PoseSearch/Source/Runtime/Public/PoseSearch/PoseSearchTrajectoryTypes.h
@@ -3,6 +3,7 @@
 #pragma once
 
 #include "Animation/AnimTypes.h"
+#include "Animation/MotionTrajectoryTypes.h"
 
 #include "PoseSearchTrajectoryTypes.generated.h"
 
@@ -21,6 +22,8 @@ struct POSESEARCH_API FPoseSearchQueryTrajectorySample
 	float AccumulatedSeconds = 0.f;
 
 	FPoseSearchQueryTrajectorySample Lerp(const FPoseSearchQueryTrajectorySample& Other, float Alpha) const;
+	void SetTransform(const FTransform& Transform);
+	FTransform GetTransform() const { return FTransform(Facing, Position); }
 };
 
 USTRUCT(BlueprintType, Category = "Motion Trajectory")
@@ -32,9 +35,11 @@ struct POSESEARCH_API FPoseSearchQueryTrajectory
 	UPROPERTY(EditAnywhere, BlueprintReadWrite, Category = "Pose Search Query Trajectory")
 	TArray<FPoseSearchQueryTrajectorySample> Samples;
 
-	FPoseSearchQueryTrajectorySample GetSampleAtTime(float Time, bool bExtrapolate) const;
+	FPoseSearchQueryTrajectorySample GetSampleAtTime(float Time, bool bExtrapolate = true) const;
+	
+	void TransformReferenceFrame(const FTransform& DeltaTransform);
 
 #if ENABLE_ANIM_DEBUG
 	void DebugDrawTrajectory(const UWorld* World, const FTransform& TransformWS) const;
 #endif // ENABLE_ANIM_DEBUG
-};
\ No newline at end of file
+};
diff --git a/Engine/Source/Runtime/Engine/Public/Animation/MotionTrajectoryTypes.h b/Engine/Source/Runtime/Engine/Public/Animation/MotionTrajectoryTypes.h
--- a/Engine/Source/Runtime/Engine/Public/Animation/MotionTrajectoryTypes.h
+++ b/Engine/Source/Runtime/Engine/Public/Animation/MotionTrajectoryTypes.h
@@ -45,7 +45,7 @@ struct ENGINE_API FTrajectorySample
 
 // A container of ordered trajectory samples and associated sampling rate
 USTRUCT(BlueprintType, Category="Motion Trajectory")
-struct ENGINE_API FTrajectorySampleRange
+struct UE_DEPRECATED(5.3, "Use FPoseSearchQueryTrajectory instead") ENGINE_API FTrajectorySampleRange
 {
 	GENERATED_BODY()
 
