motion matching - fix for Warning: Plugins/Experimental/Animation/MotionTrajectory/Source/MotionTrajectory/Private/MotionTrajectory.cpp(10): Missing '#undef LOCTEXT_NAMESPACE' for '#define LOCTEXT_NAMESPACE'

diff --git a/Engine/Plugins/Experimental/Animation/MotionTrajectory/Source/MotionTrajectory/Private/MotionTrajectory.cpp b/Engine/Plugins/Experimental/Animation/MotionTrajectory/Source/MotionTrajectory/Private/MotionTrajectory.cpp
--- a/Engine/Plugins/Experimental/Animation/MotionTrajectory/Source/MotionTrajectory/Private/MotionTrajectory.cpp
+++ b/Engine/Plugins/Experimental/Animation/MotionTrajectory/Source/MotionTrajectory/Private/MotionTrajectory.cpp
@@ -1,14 +1,8 @@
 // Copyright Epic Games, Inc. All Rights Reserved.
 
 #include "MotionTrajectory.h"
-
-#include "Engine/World.h"
-#include "GameFramework/Pawn.h"
-#include "Kismet/KismetSystemLibrary.h"
 #include "Modules/ModuleManager.h"
 
-#define LOCTEXT_NAMESPACE "MotionTrajectory"
-
 DEFINE_LOG_CATEGORY(LogMotionTrajectory);
 
 void FMotionTrajectoryModule::StartupModule()
