Add option to manually update Character Trajectory rather than depending on the OnMovementUpdated callback. This allows projects to control when the trajectory is updated (e.g. only for characters that are in view or nearby).

diff --git a/Engine/Plugins/Experimental/Animation/MotionTrajectory/Source/MotionTrajectory/Private/CharacterTrajectoryComponent.cpp b/Engine/Plugins/Experimental/Animation/MotionTrajectory/Source/MotionTrajectory/Private/CharacterTrajectoryComponent.cpp
--- a/Engine/Plugins/Experimental/Animation/MotionTrajectory/Source/MotionTrajectory/Private/CharacterTrajectoryComponent.cpp
+++ b/Engine/Plugins/Experimental/Animation/MotionTrajectory/Source/MotionTrajectory/Private/CharacterTrajectoryComponent.cpp
@@ -26,25 +26,31 @@ void UCharacterTrajectoryComponent::InitializeComponent()
 {
 	Super::InitializeComponent();
 
-	if (ACharacter* Character = Cast<ACharacter>(GetOwner()))
+	if (bAutoUpdateTrajectory)
 	{
-		Character->OnCharacterMovementUpdated.AddDynamic(this, &UCharacterTrajectoryComponent::OnMovementUpdated);
-	}
-	else
-	{
-		UE_LOG(LogMotionTrajectory, Error, TEXT("UCharacterTrajectoryComponent requires its owner to be ACharacter"));
+		if (ACharacter* Character = Cast<ACharacter>(GetOwner()))
+		{
+			Character->OnCharacterMovementUpdated.AddDynamic(this, &UCharacterTrajectoryComponent::OnMovementUpdated);
+		}
+		else
+		{
+			UE_LOG(LogMotionTrajectory, Error, TEXT("UCharacterTrajectoryComponent requires its owner to be ACharacter"));
+		}
 	}
 }
 
 void UCharacterTrajectoryComponent::UninitializeComponent()
 {
-	if (ACharacter* Character = Cast<ACharacter>(GetOwner()))
-	{
-		Character->OnCharacterMovementUpdated.RemoveDynamic(this, &UCharacterTrajectoryComponent::OnMovementUpdated);
-	}
-	else
+	if (bAutoUpdateTrajectory)
 	{
-		UE_LOG(LogMotionTrajectory, Error, TEXT("UCharacterTrajectoryComponent requires its owner to be ACharacter"));
+		if (ACharacter* Character = Cast<ACharacter>(GetOwner()))
+		{
+			Character->OnCharacterMovementUpdated.RemoveDynamic(this, &UCharacterTrajectoryComponent::OnMovementUpdated);
+		}
+		else
+		{
+			UE_LOG(LogMotionTrajectory, Error, TEXT("UCharacterTrajectoryComponent requires its owner to be ACharacter"));
+		}
 	}
 
 	Super::UninitializeComponent();
@@ -109,13 +115,18 @@ void UCharacterTrajectoryComponent::BeginPlay()
 	}
 }
 
-void UCharacterTrajectoryComponent::OnMovementUpdated(float DeltaSeconds, FVector OldLocation, FVector OldVelocity)
+void UCharacterTrajectoryComponent::UpdateTrajectory(float DeltaSeconds)
 {
 	if (!ensure(CharacterMovementComponent != nullptr && SkelMeshComponent != nullptr))
 	{
 		return;
 	}
 
+	if (!ensure(DeltaSeconds > 0.f))
+	{
+		return;
+	}
+
 	const FTransform PreviousSkelMeshComponentTransformWS = SkelMeshComponentTransformWS;
 	SkelMeshComponentTransformWS = SkelMeshComponent->GetComponentTransform();
 
@@ -135,6 +146,11 @@ void UCharacterTrajectoryComponent::OnMovementUpdated(float DeltaSeconds, FVecto
 #endif // ENABLE_ANIM_DEBUG
 }
 
+void UCharacterTrajectoryComponent::OnMovementUpdated(float DeltaSeconds, FVector OldLocation, FVector OldVelocity)
+{
+	UpdateTrajectory(DeltaSeconds);
+}
+
 void UpdateHistorySample(FPoseSearchQueryTrajectorySample& Sample, float DeltaSeconds, const FTransform& DeltaTransformCS)
 {
 	Sample.Facing = DeltaTransformCS.InverseTransformRotation(Sample.Facing);
@@ -216,7 +232,7 @@ FRotator UCharacterTrajectoryComponent::CalculateControllerRotationRate(float De
 {
 	check(GetOwner());
 
-	// OnMovementUpdated handles DeltaSeconds == 0.f, so we should never hit this.
+	// UpdateTrajectory handles DeltaSeconds == 0.f, so we should never hit this.
 	check(DeltaSeconds > 0.f);
 
 	ACharacter* CharacterOwner = Cast<ACharacter>(GetOwner());
diff --git a/Engine/Plugins/Experimental/Animation/MotionTrajectory/Source/MotionTrajectory/Public/CharacterTrajectoryComponent.h b/Engine/Plugins/Experimental/Animation/MotionTrajectory/Source/MotionTrajectory/Public/CharacterTrajectoryComponent.h
--- a/Engine/Plugins/Experimental/Animation/MotionTrajectory/Source/MotionTrajectory/Public/CharacterTrajectoryComponent.h
+++ b/Engine/Plugins/Experimental/Animation/MotionTrajectory/Source/MotionTrajectory/Public/CharacterTrajectoryComponent.h
@@ -11,7 +11,7 @@ class UCharacterMovementComponent;
 // Component for generating trajectories usable by Motion Matching. This component generates trajectories from ACharacter.
 // This is intended to provide an example and starting point for using Motion Matching with a common setup using the default UCharacterMovementComponent.
 // It is expected work flow to extend or replace this component for projects that use a custom movement component or custom movement modes.
-UCLASS(BlueprintType, meta = (BlueprintSpawnableComponent), Experimental)
+UCLASS(Blueprintable, BlueprintType, meta = (BlueprintSpawnableComponent), Experimental)
 class MOTIONTRAJECTORY_API UCharacterTrajectoryComponent : public UActorComponent
 {
 	GENERATED_BODY()
@@ -26,6 +26,10 @@ public:
 	virtual void BeginPlay() override;
 	// End UActorComponent Interface
 
+	// This can be called when bAutoUpdateTrajectory is false to manually control when trajectory updates.
+	UFUNCTION(BlueprintCallable, Category = "Trajectory")
+	void UpdateTrajectory(float DeltaSeconds);
+
 protected:
 	UFUNCTION()
 	void OnMovementUpdated(float DeltaSeconds, FVector OldLocation, FVector OldVelocity);
@@ -63,6 +67,11 @@ protected:
 	UPROPERTY(EditDefaultsOnly, BlueprintReadOnly, Category = "Trajectory Settings")
 	float RotateTowardsMovementSpeed = 10.f;
 
+	// By default the component will always update trajectory. If desired, this can be disabled and the game can choose when to update.
+	// For example, a game might want to only update trajectory for characters that are within view or very close to the local player.
+	UPROPERTY(EditDefaultsOnly, BlueprintReadOnly, Category = "Trajectory Settings")
+	bool bAutoUpdateTrajectory = true;
+
 	UPROPERTY()
 	TObjectPtr<USkeletalMeshComponent> SkelMeshComponent;
 
