motion matching - trajectory generation in world space

diff --git a/Engine/Plugins/Experimental/Animation/MotionTrajectory/Source/MotionTrajectory/Private/CharacterTrajectoryComponent.cpp b/Engine/Plugins/Experimental/Animation/MotionTrajectory/Source/MotionTrajectory/Private/CharacterTrajectoryComponent.cpp
--- a/Engine/Plugins/Experimental/Animation/MotionTrajectory/Source/MotionTrajectory/Private/CharacterTrajectoryComponent.cpp
+++ b/Engine/Plugins/Experimental/Animation/MotionTrajectory/Source/MotionTrajectory/Private/CharacterTrajectoryComponent.cpp
@@ -78,8 +78,6 @@ void UCharacterTrajectoryComponent::BeginPlay()
 		return;
 	}
 
-	SkelMeshComponentTransformWS = SkelMeshComponent->GetComponentTransform();
-
 	// Default forward in the engine is the X axis, but data often diverges from this (e.g. it's common for skeletal meshes to be Y forward).
 	// We determine the forward direction in the space of the skeletal mesh component based on the offset from the actor.
 	ForwardFacingCS = SkelMeshComponent->GetRelativeRotation().Quaternion().Inverse();
@@ -127,21 +125,16 @@ void UCharacterTrajectoryComponent::UpdateTrajectory(float DeltaSeconds)
 		return;
 	}
 
-	const FTransform PreviousSkelMeshComponentTransformWS = SkelMeshComponentTransformWS;
-	SkelMeshComponentTransformWS = SkelMeshComponent->GetComponentTransform();
-
-	const FTransform SkelMeshTransformDelta = SkelMeshComponentTransformWS.GetRelativeTransform(PreviousSkelMeshComponentTransformWS);
-	UpdateHistory(DeltaSeconds, SkelMeshTransformDelta);
-
-	const FVector VelocityCS = SkelMeshComponentTransformWS.InverseTransformVectorNoScale(CharacterMovementComponent->Velocity);
-	const FVector AccelerationCS = SkelMeshComponentTransformWS.InverseTransformVectorNoScale(CharacterMovementComponent->GetCurrentAcceleration());
+	const FTransform& SkelMeshComponentTransformWS = SkelMeshComponent->GetComponentTransform();
 	const FRotator ControllerRotationRate = CalculateControllerRotationRate(DeltaSeconds, CharacterMovementComponent->ShouldRemainVertical());
-	UpdatePrediction(VelocityCS, AccelerationCS, ControllerRotationRate);
+		
+	UpdatePrediction(SkelMeshComponentTransformWS.GetTranslation(), SkelMeshComponentTransformWS.GetRotation(), CharacterMovementComponent->Velocity, CharacterMovementComponent->GetCurrentAcceleration(), ControllerRotationRate);
+	UpdateHistory(DeltaSeconds);
 
 #if ENABLE_ANIM_DEBUG
 	if (CVarCharacterTrajectoryDebug.GetValueOnAnyThread())
 	{
-		Trajectory.DebugDrawTrajectory(GetWorld(), SkelMeshComponentTransformWS);
+		Trajectory.DebugDrawTrajectory(GetWorld());
 	}
 #endif // ENABLE_ANIM_DEBUG
 }
@@ -151,20 +144,9 @@ void UCharacterTrajectoryComponent::OnMovementUpdated(float DeltaSeconds, FVecto
 	UpdateTrajectory(DeltaSeconds);
 }
 
-void UpdateHistorySample(FPoseSearchQueryTrajectorySample& Sample, float DeltaSeconds, const FTransform& DeltaTransformCS)
-{
-	Sample.Facing = DeltaTransformCS.InverseTransformRotation(Sample.Facing);
-	Sample.Position = DeltaTransformCS.InverseTransformPosition(Sample.Position);
-	Sample.AccumulatedSeconds -= DeltaSeconds;
-}
-
-// This function moves each history sample by the inverse of the character's current motion (i.e. if the character is moving forward, the history
-// samples move backward). It also shifts the range of history samples whenever a new history sample should be recorded.
-// This allows us to keep a single sample array in component space that can be read directly by the Motion Matching node, rather than storing world
-// transforms in a separate list and converting them to component space each update. 
-// This also allows us to create "faked" trajectories that match animation data rather than the simulation (e.g. if our animation data only has coverage
-// for one speed, we can adjust the history by a single speed to produce trajectories that best match the data).
-void UCharacterTrajectoryComponent::UpdateHistory(float DeltaSeconds, const FTransform& DeltaTransformCS)
+// This function shifts the range of history samples whenever a new history sample should be recorded.
+// This allows us to keep a single sample array in world space that can be read directly by the Motion Matching node. 
+void UCharacterTrajectoryComponent::UpdateHistory(float DeltaSeconds)
 {
 	check(NumHistorySamples <= Trajectory.Samples.Num());
 
@@ -174,54 +156,61 @@ void UCharacterTrajectoryComponent::UpdateHistory(float DeltaSeconds, const FTra
 		for (int32 Index = 0; Index < NumHistorySamples; ++Index)
 		{
 			Trajectory.Samples[Index] = Trajectory.Samples[Index + 1];
-
-			UpdateHistorySample(Trajectory.Samples[Index], DeltaSeconds, DeltaTransformCS);
+			Trajectory.Samples[Index].AccumulatedSeconds -= DeltaSeconds;
 		}
 	}
 	else
 	{
 		for (int32 Index = 0; Index < NumHistorySamples; ++Index)
 		{
-			UpdateHistorySample(Trajectory.Samples[Index], DeltaSeconds, DeltaTransformCS);
+			Trajectory.Samples[Index].AccumulatedSeconds -= DeltaSeconds;
 		}
 	}
 }
 
-void UCharacterTrajectoryComponent::UpdatePrediction(const FVector& VelocityCS, const FVector& AccelerationCS, const FRotator& ControllerRotationRate)
+void UCharacterTrajectoryComponent::UpdatePrediction(const FVector& PositionWS, const FQuat& FacingWS, const FVector& VelocityWS, const FVector& AccelerationWS, const FRotator& ControllerRotationRate)
 {
 	check(CharacterMovementComponent);
 
-	FVector CurrentPositionCS = FVector::ZeroVector;
-	FVector CurrentVelocityCS = VelocityCS;
-	FVector CurrentAccelerationCS = AccelerationCS;
-	FQuat CurrentFacingCS = ForwardFacingCS;
+	FVector CurrentPositionWS = PositionWS;
+	FVector CurrentVelocityWS = VelocityWS;
+	FVector CurrentAccelerationWS = AccelerationWS;
+	FQuat CurrentFacingWS = FacingWS;
 	float AccumulatedSeconds = 0.f;
 
 	FQuat ControllerRotationPerStep = (ControllerRotationRate * SecondsPerPredictionSample).Quaternion();
 
-	for (int32 Index = NumHistorySamples + 1; Index < Trajectory.Samples.Num(); ++Index)
+	const int32 LastIndex = Trajectory.Samples.Num() - 1;
+	if (NumHistorySamples <= LastIndex)
 	{
-		CurrentPositionCS += CurrentVelocityCS * SecondsPerPredictionSample;
-		AccumulatedSeconds += SecondsPerPredictionSample;
-
-		// Account for the controller (e.g. the camera) rotating.
-		CurrentFacingCS = ControllerRotationPerStep * CurrentFacingCS;
-		CurrentAccelerationCS = ControllerRotationPerStep * CurrentAccelerationCS;
-
-		Trajectory.Samples[Index].Position = CurrentPositionCS;
-		Trajectory.Samples[Index].Facing = CurrentFacingCS;
-		Trajectory.Samples[Index].AccumulatedSeconds = AccumulatedSeconds;
-
-		FVector NewVelocityCS = FVector::ZeroVector;
-		UCharacterMovementTrajectoryLibrary::StepCharacterMovementGroundPrediction(
-			SecondsPerPredictionSample, CurrentVelocityCS, CurrentAccelerationCS, CharacterMovementComponent,
-			NewVelocityCS);
-		CurrentVelocityCS = NewVelocityCS;
-
-		if (CharacterMovementComponent->bOrientRotationToMovement && !CurrentAccelerationCS.IsNearlyZero())
+		for (int32 Index = NumHistorySamples; ; ++Index)
 		{
-			// Rotate towards acceleration.
-			CurrentFacingCS = FMath::QInterpConstantTo(CurrentFacingCS, CurrentAccelerationCS.ToOrientationQuat(), SecondsPerPredictionSample, RotateTowardsMovementSpeed);
+			Trajectory.Samples[Index].Position = CurrentPositionWS;
+			Trajectory.Samples[Index].Facing = CurrentFacingWS;
+			Trajectory.Samples[Index].AccumulatedSeconds = AccumulatedSeconds;
+
+			if (Index == LastIndex)
+			{
+				break;
+			}
+
+			CurrentPositionWS += CurrentVelocityWS * SecondsPerPredictionSample;
+			AccumulatedSeconds += SecondsPerPredictionSample;
+
+			// Account for the controller (e.g. the camera) rotating.
+			CurrentFacingWS = ControllerRotationPerStep * CurrentFacingWS;
+			CurrentAccelerationWS = ControllerRotationPerStep * CurrentAccelerationWS;
+
+			FVector NewVelocityCS = FVector::ZeroVector;
+			UCharacterMovementTrajectoryLibrary::StepCharacterMovementGroundPrediction(SecondsPerPredictionSample, CurrentVelocityWS, CurrentAccelerationWS, CharacterMovementComponent, NewVelocityCS);
+			CurrentVelocityWS = NewVelocityCS;
+
+			if (CharacterMovementComponent->bOrientRotationToMovement && !CurrentAccelerationWS.IsNearlyZero())
+			{
+				// Rotate towards acceleration.
+				const FVector CurrentAccelerationCS = SkelMeshComponent->GetRelativeRotation().Quaternion().RotateVector(CurrentAccelerationWS);
+				CurrentFacingWS = FMath::QInterpConstantTo(CurrentFacingWS, CurrentAccelerationCS.ToOrientationQuat(), SecondsPerPredictionSample, RotateTowardsMovementSpeed);
+			}
 		}
 	}
 }
diff --git a/Engine/Plugins/Experimental/Animation/MotionTrajectory/Source/MotionTrajectory/Public/CharacterTrajectoryComponent.h b/Engine/Plugins/Experimental/Animation/MotionTrajectory/Source/MotionTrajectory/Public/CharacterTrajectoryComponent.h
--- a/Engine/Plugins/Experimental/Animation/MotionTrajectory/Source/MotionTrajectory/Public/CharacterTrajectoryComponent.h
+++ b/Engine/Plugins/Experimental/Animation/MotionTrajectory/Source/MotionTrajectory/Public/CharacterTrajectoryComponent.h
@@ -34,8 +34,8 @@ protected:
 	UFUNCTION()
 	void OnMovementUpdated(float DeltaSeconds, FVector OldLocation, FVector OldVelocity);
 
-	void UpdateHistory(float DeltaSeconds, const FTransform& DeltaTransformCS);
-	void UpdatePrediction(const FVector& VelocityCS, const FVector& AccelerationCS, const FRotator& ControllerRotationRate);
+	void UpdateHistory(float DeltaSeconds);
+	void UpdatePrediction(const FVector& PositionWS, const FQuat& FacingWS, const FVector& VelocityWS, const FVector& AccelerationWS, const FRotator& ControllerRotationRate);
 
 	FRotator CalculateControllerRotationRate(float DeltaSeconds, bool bShouldRemainVertical);
 
@@ -82,9 +82,6 @@ protected:
 	float SecondsPerHistorySample = 0.f;
 	float SecondsPerPredictionSample = 0.f;
 
-	// Current transform of the skeletal mesh component, used to calculate the movement delta between frames.
-	FTransform SkelMeshComponentTransformWS = FTransform::Identity;
-
 	// Forward axis for the SkeletalMeshComponent. It's common for skeletal mesh and animation data to not be X forward.
 	FQuat ForwardFacingCS = FQuat::Identity;
 
diff --git a/Engine/Plugins/Experimental/Animation/PoseSearch/Source/Runtime/Private/PoseSearchLibrary.cpp b/Engine/Plugins/Experimental/Animation/PoseSearch/Source/Runtime/Private/PoseSearchLibrary.cpp
--- a/Engine/Plugins/Experimental/Animation/PoseSearch/Source/Runtime/Private/PoseSearchLibrary.cpp
+++ b/Engine/Plugins/Experimental/Animation/PoseSearch/Source/Runtime/Private/PoseSearchLibrary.cpp
@@ -415,12 +415,16 @@ void UPoseSearchLibrary::UpdateMotionMatchingState(
 #endif
 }
 
-// transforms Trajectory from the SkeletalMeshComponent relative space into Character relative space, and scale it by TrajectorySpeedMultiplier
+// transforms Trajectory from world space to mesh component space, and scale it by TrajectorySpeedMultiplier
 FPoseSearchQueryTrajectory UPoseSearchLibrary::ProcessTrajectory(const FPoseSearchQueryTrajectory& Trajectory, const FTransform& OwnerTransformWS, const FTransform& ComponentTransformWS, float TrajectorySpeedMultiplier)
 {
-	const FTransform ReferenceChangeTransform = OwnerTransformWS.GetRelativeTransform(ComponentTransformWS);
 	FPoseSearchQueryTrajectory TrajectoryCS = Trajectory;
-	TrajectoryCS.TransformReferenceFrame(ReferenceChangeTransform);
+	const FTransform ToComponentTransform = ComponentTransformWS.Inverse();
+	for (FPoseSearchQueryTrajectorySample& Sample : TrajectoryCS.Samples)
+	{
+		Sample.Position = ToComponentTransform.TransformPosition(Sample.Position);
+		Sample.Facing = ToComponentTransform.TransformRotation(Sample.Facing);
+	}
 
 	if (!FMath::IsNearlyEqual(TrajectorySpeedMultiplier, 1.f) && !FMath::IsNearlyZero(TrajectorySpeedMultiplier))
 	{
@@ -429,6 +433,7 @@ FPoseSearchQueryTrajectory UPoseSearchLibrary::ProcessTrajectory(const FPoseSear
 			Sample.AccumulatedSeconds /= TrajectorySpeedMultiplier;
 		}
 	}
+
 	return TrajectoryCS;
 }
 
diff --git a/Engine/Plugins/Experimental/Animation/PoseSearch/Source/Runtime/Private/PoseSearchTrajectoryTypes.cpp b/Engine/Plugins/Experimental/Animation/PoseSearch/Source/Runtime/Private/PoseSearchTrajectoryTypes.cpp
--- a/Engine/Plugins/Experimental/Animation/PoseSearch/Source/Runtime/Private/PoseSearchTrajectoryTypes.cpp
+++ b/Engine/Plugins/Experimental/Animation/PoseSearch/Source/Runtime/Private/PoseSearchTrajectoryTypes.cpp
@@ -59,46 +59,24 @@ FPoseSearchQueryTrajectorySample FPoseSearchQueryTrajectory::GetSampleAtTime(flo
 	return FPoseSearchQueryTrajectorySample();
 }
 
-void FPoseSearchQueryTrajectory::TransformReferenceFrame(const FTransform& DeltaTransform)
-{
-	const FTransform InverseDeltaTransform = DeltaTransform.Inverse();
-	for (FPoseSearchQueryTrajectorySample& Sample : Samples)
-	{
-		const FTransform Transform = InverseDeltaTransform * Sample.GetTransform() * DeltaTransform;
-		Sample.SetTransform(Transform);
-	}
-}
-
 #if ENABLE_ANIM_DEBUG
-void FPoseSearchQueryTrajectory::DebugDrawTrajectory(const UWorld* World, const FTransform& TransformWS) const
+void FPoseSearchQueryTrajectory::DebugDrawTrajectory(const UWorld* World) const
 {
-	for (int32 Index = 0; Index < Samples.Num(); ++Index)
+	const int32 LastIndex = Samples.Num() - 1;
+	if (LastIndex >= 0)
 	{
-		const FVector CurrentSamplePositionWS = TransformWS.TransformPosition(Samples[Index].Position);
-
-		DrawDebugSphere(
-			World,
-			CurrentSamplePositionWS,
-			2.f /*Radius*/, 4 /*Segments*/,
-			FColor::Black, false /*bPersistentLines*/, -1.f /*LifeTime*/, 0 /*DepthPriority*/, 1.f /*Thickness*/);
-
-		if (Samples.IsValidIndex(Index + 1))
+		for (int32 Index = 0; ; ++Index)
 		{
-			const FVector NextSamplePositionWS = TransformWS.TransformPosition(Samples[Index + 1].Position);
+			DrawDebugSphere(World, Samples[Index].Position, 2.f /*Radius*/, 4 /*Segments*/, FColor::Black);
+			DrawDebugCoordinateSystem(World, Samples[Index].Position, FRotator(Samples[Index].Facing), 12.f /*Scale*/);
 
-			DrawDebugLine(
-				World,
-				CurrentSamplePositionWS,
-				NextSamplePositionWS,
-				FColor::Black, false /*bPersistentLines*/, -1.f /*LifeTime*/, 0 /*DepthPriority*/, 1.f /*Thickness*/);
+			if (Index == LastIndex)
+			{
+				break;
+			}
+			
+			DrawDebugLine(World, Samples[Index].Position, Samples[Index + 1].Position, FColor::Black);
 		}
-
-		const FQuat CurrentSampleFacingWS = TransformWS.TransformRotation(Samples[Index].Facing);
-		DrawDebugDirectionalArrow(
-			World,
-			CurrentSamplePositionWS,
-			CurrentSamplePositionWS + CurrentSampleFacingWS.RotateVector(FVector::ForwardVector) * 25.f,
-			20.f, FColor::Orange, false /*bPersistentLines*/, -1.f /*LifeTime*/, 0 /*DepthPriority*/, 1.f /*Thickness*/);
 	}
 }
 #endif // ENABLE_ANIM_DEBUG
\ No newline at end of file
diff --git a/Engine/Plugins/Experimental/Animation/PoseSearch/Source/Runtime/Public/PoseSearch/PoseSearchTrajectoryTypes.h b/Engine/Plugins/Experimental/Animation/PoseSearch/Source/Runtime/Public/PoseSearch/PoseSearchTrajectoryTypes.h
--- a/Engine/Plugins/Experimental/Animation/PoseSearch/Source/Runtime/Public/PoseSearch/PoseSearchTrajectoryTypes.h
+++ b/Engine/Plugins/Experimental/Animation/PoseSearch/Source/Runtime/Public/PoseSearch/PoseSearchTrajectoryTypes.h
@@ -37,9 +37,7 @@ struct POSESEARCH_API FPoseSearchQueryTrajectory
 
 	FPoseSearchQueryTrajectorySample GetSampleAtTime(float Time, bool bExtrapolate = true) const;
 	
-	void TransformReferenceFrame(const FTransform& DeltaTransform);
-
 #if ENABLE_ANIM_DEBUG
-	void DebugDrawTrajectory(const UWorld* World, const FTransform& TransformWS) const;
+	void DebugDrawTrajectory(const UWorld* World) const;
 #endif // ENABLE_ANIM_DEBUG
 };
