motion matching - UCharacterTrajectoryComponent update history fix

diff --git a/Engine/Plugins/Experimental/Animation/MotionTrajectory/Source/MotionTrajectory/Private/CharacterTrajectoryComponent.cpp b/Engine/Plugins/Experimental/Animation/MotionTrajectory/Source/MotionTrajectory/Private/CharacterTrajectoryComponent.cpp
--- a/Engine/Plugins/Experimental/Animation/MotionTrajectory/Source/MotionTrajectory/Private/CharacterTrajectoryComponent.cpp
+++ b/Engine/Plugins/Experimental/Animation/MotionTrajectory/Source/MotionTrajectory/Private/CharacterTrajectoryComponent.cpp
@@ -97,6 +97,8 @@ void UCharacterTrajectoryComponent::BeginPlay()
 	DefaultSample.Position = FVector::ZeroVector;
 	DefaultSample.AccumulatedSeconds = 0.f;
 
+	const FTransform& SkelMeshComponentTransformWS = SkelMeshComponent->GetComponentTransform();
+
 	// History + current sample + prediction
 	Trajectory.Samples.Init(DefaultSample, NumHistorySamples + 1 + NumPredictionSamples);
 
@@ -104,12 +106,16 @@ void UCharacterTrajectoryComponent::BeginPlay()
 	for (int32 i = 0; i < NumHistorySamples; ++i)
 	{
 		Trajectory.Samples[i].AccumulatedSeconds = SecondsPerHistorySample * (i - NumHistorySamples);
+		Trajectory.Samples[i].Position = SkelMeshComponentTransformWS.GetTranslation();
+		Trajectory.Samples[i].Facing = SkelMeshComponentTransformWS.GetRotation();
 	}
 
 	// initializing history samples AccumulatedSeconds
 	for (int32 i = NumHistorySamples + 1; i < Trajectory.Samples.Num(); ++i)
 	{
 		Trajectory.Samples[i].AccumulatedSeconds = SecondsPerPredictionSample * (i - NumHistorySamples);
+		Trajectory.Samples[i].Position = SkelMeshComponentTransformWS.GetTranslation();
+		Trajectory.Samples[i].Facing = SkelMeshComponentTransformWS.GetRotation();
 	}
 }
 
@@ -128,8 +134,8 @@ void UCharacterTrajectoryComponent::UpdateTrajectory(float DeltaSeconds)
 	const FTransform& SkelMeshComponentTransformWS = SkelMeshComponent->GetComponentTransform();
 	const FRotator ControllerRotationRate = CalculateControllerRotationRate(DeltaSeconds, CharacterMovementComponent->ShouldRemainVertical());
 		
-	UpdatePrediction(SkelMeshComponentTransformWS.GetTranslation(), SkelMeshComponentTransformWS.GetRotation(), CharacterMovementComponent->Velocity, CharacterMovementComponent->GetCurrentAcceleration(), ControllerRotationRate);
 	UpdateHistory(DeltaSeconds);
+	UpdatePrediction(SkelMeshComponentTransformWS.GetTranslation(), SkelMeshComponentTransformWS.GetRotation(), CharacterMovementComponent->Velocity, CharacterMovementComponent->GetCurrentAcceleration(), ControllerRotationRate);
 
 #if ENABLE_ANIM_DEBUG
 	if (CVarCharacterTrajectoryDebug.GetValueOnAnyThread())
