motion trajectory - MaxControllerRotationRate

diff --git a/Engine/Plugins/Experimental/Animation/MotionTrajectory/Source/MotionTrajectory/Private/CharacterTrajectoryComponent.cpp b/Engine/Plugins/Experimental/Animation/MotionTrajectory/Source/MotionTrajectory/Private/CharacterTrajectoryComponent.cpp
--- a/Engine/Plugins/Experimental/Animation/MotionTrajectory/Source/MotionTrajectory/Private/CharacterTrajectoryComponent.cpp
+++ b/Engine/Plugins/Experimental/Animation/MotionTrajectory/Source/MotionTrajectory/Private/CharacterTrajectoryComponent.cpp
@@ -248,5 +248,13 @@ FRotator UCharacterTrajectoryComponent::CalculateControllerRotationRate(float De
 	const FRotator DesiredRotationDelta = DesiredControllerRotation - DesiredControllerRotationLastUpdate;
 	DesiredControllerRotationLastUpdate = DesiredControllerRotation;
 
-	return DesiredRotationDelta.GetNormalized() * (1.f / DeltaSeconds);
+	FRotator ControllerRotationRate = DesiredRotationDelta.GetNormalized() * (1.f / DeltaSeconds);
+	if (MaxControllerRotationRate >= 0.f)
+	{
+		ControllerRotationRate.Pitch = FMath::Sign(ControllerRotationRate.Pitch) * FMath::Min(FMath::Abs(ControllerRotationRate.Pitch), MaxControllerRotationRate);
+		ControllerRotationRate.Yaw = FMath::Sign(ControllerRotationRate.Yaw) * FMath::Min(FMath::Abs(ControllerRotationRate.Yaw), MaxControllerRotationRate);
+		ControllerRotationRate.Roll = FMath::Sign(ControllerRotationRate.Roll) * FMath::Min(FMath::Abs(ControllerRotationRate.Roll), MaxControllerRotationRate);
+	}
+
+	return ControllerRotationRate;
 }
\ No newline at end of file
diff --git a/Engine/Plugins/Experimental/Animation/MotionTrajectory/Source/MotionTrajectory/Public/CharacterTrajectoryComponent.h b/Engine/Plugins/Experimental/Animation/MotionTrajectory/Source/MotionTrajectory/Public/CharacterTrajectoryComponent.h
--- a/Engine/Plugins/Experimental/Animation/MotionTrajectory/Source/MotionTrajectory/Public/CharacterTrajectoryComponent.h
+++ b/Engine/Plugins/Experimental/Animation/MotionTrajectory/Source/MotionTrajectory/Public/CharacterTrajectoryComponent.h
@@ -72,6 +72,11 @@ protected:
 	UPROPERTY(EditDefaultsOnly, BlueprintReadOnly, Category = "Trajectory Settings")
 	bool bAutoUpdateTrajectory = true;
 
+	// Maximum controller rotation rate in degrees per second used to clamp the character owner controller desired rotation to generate the prediction trajectory.
+	// Negative values disable the clamping behavior
+	UPROPERTY(EditDefaultsOnly, BlueprintReadOnly, Category = "Trajectory Settings")
+	float MaxControllerRotationRate = -1.f;	
+
 	UPROPERTY()
 	TObjectPtr<USkeletalMeshComponent> SkelMeshComponent;
 
