Motion Matching: Added BlueprintNativeEvent to override mesh facing for trajectory calculation.
#rnx

diff --git a/Engine/Plugins/Experimental/Animation/MotionTrajectory/Source/MotionTrajectory/Private/CharacterTrajectoryComponent.cpp b/Engine/Plugins/Experimental/Animation/MotionTrajectory/Source/MotionTrajectory/Private/CharacterTrajectoryComponent.cpp
--- a/Engine/Plugins/Experimental/Animation/MotionTrajectory/Source/MotionTrajectory/Private/CharacterTrajectoryComponent.cpp
+++ b/Engine/Plugins/Experimental/Animation/MotionTrajectory/Source/MotionTrajectory/Private/CharacterTrajectoryComponent.cpp
@@ -56,6 +56,12 @@ void UCharacterTrajectoryComponent::UninitializeComponent()
 	Super::UninitializeComponent();
 }
 
+FRotator UCharacterTrajectoryComponent::GetFacingFromMeshComponent_Implementation(const USkeletalMeshComponent* MeshComponent) const
+{
+	check(MeshComponent);
+	return MeshComponent->GetComponentRotation();
+}
+
 void UCharacterTrajectoryComponent::BeginPlay()
 {
 	Super::BeginPlay();
@@ -97,7 +103,8 @@ void UCharacterTrajectoryComponent::BeginPlay()
 	DefaultSample.Position = FVector::ZeroVector;
 	DefaultSample.AccumulatedSeconds = 0.f;
 
-	const FTransform& SkelMeshComponentTransformWS = SkelMeshComponent->GetComponentTransform();
+	const FVector SkelMeshComponentLocationWS = SkelMeshComponent->GetComponentLocation();
+	const FQuat FacingWS = FQuat(GetFacingFromMeshComponent(SkelMeshComponent));
 
 	// History + current sample + prediction
 	Trajectory.Samples.Init(DefaultSample, NumHistorySamples + 1 + NumPredictionSamples);
@@ -106,16 +113,16 @@ void UCharacterTrajectoryComponent::BeginPlay()
 	for (int32 i = 0; i < NumHistorySamples; ++i)
 	{
 		Trajectory.Samples[i].AccumulatedSeconds = SecondsPerHistorySample * (i - NumHistorySamples);
-		Trajectory.Samples[i].Position = SkelMeshComponentTransformWS.GetTranslation();
-		Trajectory.Samples[i].Facing = SkelMeshComponentTransformWS.GetRotation();
+		Trajectory.Samples[i].Position = SkelMeshComponentLocationWS;
+		Trajectory.Samples[i].Facing = FacingWS;
 	}
 
 	// initializing history samples AccumulatedSeconds
 	for (int32 i = NumHistorySamples + 1; i < Trajectory.Samples.Num(); ++i)
 	{
 		Trajectory.Samples[i].AccumulatedSeconds = SecondsPerPredictionSample * (i - NumHistorySamples);
-		Trajectory.Samples[i].Position = SkelMeshComponentTransformWS.GetTranslation();
-		Trajectory.Samples[i].Facing = SkelMeshComponentTransformWS.GetRotation();
+		Trajectory.Samples[i].Position = SkelMeshComponentLocationWS;
+		Trajectory.Samples[i].Facing = FacingWS;
 	}
 }
 
@@ -131,11 +138,12 @@ void UCharacterTrajectoryComponent::UpdateTrajectory(float DeltaSeconds)
 		return;
 	}
 
-	const FTransform& SkelMeshComponentTransformWS = SkelMeshComponent->GetComponentTransform();
+	const FVector SkelMeshComponentLocationWS = SkelMeshComponent->GetComponentLocation();
+	const FQuat FacingWS = FQuat(GetFacingFromMeshComponent(SkelMeshComponent));
 	const FRotator ControllerRotationRate = CalculateControllerRotationRate(DeltaSeconds, CharacterMovementComponent->ShouldRemainVertical());
 		
 	UpdateHistory(DeltaSeconds);
-	UpdatePrediction(SkelMeshComponentTransformWS.GetTranslation(), SkelMeshComponentTransformWS.GetRotation(), CharacterMovementComponent->Velocity, CharacterMovementComponent->GetCurrentAcceleration(), ControllerRotationRate);
+	UpdatePrediction(SkelMeshComponentLocationWS, FacingWS, CharacterMovementComponent->Velocity, CharacterMovementComponent->GetCurrentAcceleration(), ControllerRotationRate);
 
 #if ENABLE_ANIM_DEBUG
 	if (CVarCharacterTrajectoryDebug.GetValueOnAnyThread())
diff --git a/Engine/Plugins/Experimental/Animation/MotionTrajectory/Source/MotionTrajectory/Public/CharacterTrajectoryComponent.h b/Engine/Plugins/Experimental/Animation/MotionTrajectory/Source/MotionTrajectory/Public/CharacterTrajectoryComponent.h
--- a/Engine/Plugins/Experimental/Animation/MotionTrajectory/Source/MotionTrajectory/Public/CharacterTrajectoryComponent.h
+++ b/Engine/Plugins/Experimental/Animation/MotionTrajectory/Source/MotionTrajectory/Public/CharacterTrajectoryComponent.h
@@ -30,6 +30,11 @@ public:
 	UFUNCTION(BlueprintCallable, Category = "Trajectory")
 	void UpdateTrajectory(float DeltaSeconds);
 
+	// This can be used to override the facing used for trajectory calculation.
+	UFUNCTION(BlueprintNativeEvent, Category = "Trajectory")
+	FRotator GetFacingFromMeshComponent(const USkeletalMeshComponent* MeshComponent) const;
+	virtual FRotator GetFacingFromMeshComponent_Implementation(const USkeletalMeshComponent* MeshComponent) const;
+
 protected:
 	UFUNCTION()
 	void OnMovementUpdated(float DeltaSeconds, FVector OldLocation, FVector OldVelocity);
