motion trajectory - adding SpeedRemappingCurve and AccelerationRemappingCurve to UpdatePrediction

diff --git a/Engine/Plugins/Experimental/Animation/MotionTrajectory/Source/MotionTrajectory/Private/CharacterTrajectoryComponent.cpp b/Engine/Plugins/Experimental/Animation/MotionTrajectory/Source/MotionTrajectory/Private/CharacterTrajectoryComponent.cpp
--- a/Engine/Plugins/Experimental/Animation/MotionTrajectory/Source/MotionTrajectory/Private/CharacterTrajectoryComponent.cpp
+++ b/Engine/Plugins/Experimental/Animation/MotionTrajectory/Source/MotionTrajectory/Private/CharacterTrajectoryComponent.cpp
@@ -182,13 +182,28 @@ void UCharacterTrajectoryComponent::UpdateHistory(float DeltaSeconds)
 	}
 }
 
+static FVector RemapVectorMagnitudeWithCurve(const FVector& Vector, bool bUseCurve, const FRuntimeFloatCurve& Curve)
+{
+	if (bUseCurve)
+	{
+		const float Length = Vector.Length();
+		if (Length > UE_KINDA_SMALL_NUMBER)
+		{
+			const float RemappedLength = Curve.GetRichCurveConst()->Eval(Length);
+			return Vector * (RemappedLength / Length);
+		}
+	}
+
+	return Vector;
+}
+
 void UCharacterTrajectoryComponent::UpdatePrediction(const FVector& PositionWS, const FQuat& FacingWS, const FVector& VelocityWS, const FVector& AccelerationWS, const FRotator& ControllerRotationRate)
 {
 	check(CharacterMovementComponent);
 
 	FVector CurrentPositionWS = PositionWS;
-	FVector CurrentVelocityWS = VelocityWS;
-	FVector CurrentAccelerationWS = AccelerationWS;
+	FVector CurrentVelocityWS = RemapVectorMagnitudeWithCurve(VelocityWS, bUseSpeedRemappingCurve, SpeedRemappingCurve);
+	FVector CurrentAccelerationWS = RemapVectorMagnitudeWithCurve(AccelerationWS, bUseAccelerationRemappingCurve, AccelerationRemappingCurve);
 	FQuat CurrentFacingWS = FacingWS;
 	float AccumulatedSeconds = 0.f;
 
@@ -213,11 +228,11 @@ void UCharacterTrajectoryComponent::UpdatePrediction(const FVector& PositionWS,
 
 			// Account for the controller (e.g. the camera) rotating.
 			CurrentFacingWS = ControllerRotationPerStep * CurrentFacingWS;
-			CurrentAccelerationWS = ControllerRotationPerStep * CurrentAccelerationWS;
+			CurrentAccelerationWS = RemapVectorMagnitudeWithCurve(ControllerRotationPerStep * CurrentAccelerationWS, bUseAccelerationRemappingCurve, AccelerationRemappingCurve);
 
 			FVector NewVelocityCS = FVector::ZeroVector;
 			UCharacterMovementTrajectoryLibrary::StepCharacterMovementGroundPrediction(SecondsPerPredictionSample, CurrentVelocityWS, CurrentAccelerationWS, CharacterMovementComponent, NewVelocityCS);
-			CurrentVelocityWS = NewVelocityCS;
+			CurrentVelocityWS = RemapVectorMagnitudeWithCurve(NewVelocityCS, bUseSpeedRemappingCurve, SpeedRemappingCurve);
 
 			if (CharacterMovementComponent->bOrientRotationToMovement && !CurrentAccelerationWS.IsNearlyZero())
 			{
diff --git a/Engine/Plugins/Experimental/Animation/MotionTrajectory/Source/MotionTrajectory/Public/CharacterTrajectoryComponent.h b/Engine/Plugins/Experimental/Animation/MotionTrajectory/Source/MotionTrajectory/Public/CharacterTrajectoryComponent.h
--- a/Engine/Plugins/Experimental/Animation/MotionTrajectory/Source/MotionTrajectory/Public/CharacterTrajectoryComponent.h
+++ b/Engine/Plugins/Experimental/Animation/MotionTrajectory/Source/MotionTrajectory/Public/CharacterTrajectoryComponent.h
@@ -3,6 +3,7 @@
 #pragma once
 
 #include "Components/ActorComponent.h"
+#include "Curves/CurveFloat.h"
 #include "PoseSearch/PoseSearchTrajectoryTypes.h"
 #include "CharacterTrajectoryComponent.generated.h"
 
@@ -79,9 +80,21 @@ protected:
 
 	// Maximum controller rotation rate in degrees per second used to clamp the character owner controller desired rotation to generate the prediction trajectory.
 	// Negative values disable the clamping behavior
-	UPROPERTY(EditDefaultsOnly, BlueprintReadOnly, Category = "Trajectory Settings")
+	UPROPERTY(EditDefaultsOnly, Category = "Trajectory Settings")
 	float MaxControllerRotationRate = -1.f;	
 
+	UPROPERTY(EditDefaultsOnly, Category = "Trajectory Settings", meta = (InlineEditConditionToggle))
+	bool bUseSpeedRemappingCurve = false;
+
+	UPROPERTY(EditDefaultsOnly, Category = "Trajectory Settings", meta = (EditCondition = "bUseSpeedRemappingCurve"))
+	FRuntimeFloatCurve SpeedRemappingCurve;
+
+	UPROPERTY(EditDefaultsOnly, Category = "Trajectory Settings", meta = (InlineEditConditionToggle))
+	bool bUseAccelerationRemappingCurve = false;
+
+	UPROPERTY(EditDefaultsOnly, Category = "Trajectory Settings", meta = (EditCondition = "bUseAccelerationRemappingCurve"))
+	FRuntimeFloatCurve AccelerationRemappingCurve;
+
 	UPROPERTY()
 	TObjectPtr<USkeletalMeshComponent> SkelMeshComponent;
 
