Guard against trajectory component being updated twice in the same frame.

diff --git a/Engine/Plugins/Experimental/Animation/MotionTrajectory/Source/MotionTrajectory/Private/CharacterTrajectoryComponent.cpp b/Engine/Plugins/Experimental/Animation/MotionTrajectory/Source/MotionTrajectory/Private/CharacterTrajectoryComponent.cpp
--- a/Engine/Plugins/Experimental/Animation/MotionTrajectory/Source/MotionTrajectory/Private/CharacterTrajectoryComponent.cpp
+++ b/Engine/Plugins/Experimental/Animation/MotionTrajectory/Source/MotionTrajectory/Private/CharacterTrajectoryComponent.cpp
@@ -128,6 +128,11 @@ void UCharacterTrajectoryComponent::BeginPlay()
 
 void UCharacterTrajectoryComponent::UpdateTrajectory(float DeltaSeconds)
 {
+	if (LastUpdateFrameNumber != 0 && LastUpdateFrameNumber == GFrameNumber)
+	{
+		return;
+	}
+
 	if (!ensure(CharacterMovementComponent != nullptr && SkelMeshComponent != nullptr))
 	{
 		return;
@@ -145,6 +150,8 @@ void UCharacterTrajectoryComponent::UpdateTrajectory(float DeltaSeconds)
 	UpdateHistory(DeltaSeconds);
 	UpdatePrediction(SkelMeshComponentLocationWS, FacingWS, CharacterMovementComponent->Velocity, CharacterMovementComponent->GetCurrentAcceleration(), ControllerRotationRate);
 
+	LastUpdateFrameNumber = GFrameNumber;
+
 #if ENABLE_ANIM_DEBUG
 	if (CVarCharacterTrajectoryDebug.GetValueOnAnyThread())
 	{
diff --git a/Engine/Plugins/Experimental/Animation/MotionTrajectory/Source/MotionTrajectory/Public/CharacterTrajectoryComponent.h b/Engine/Plugins/Experimental/Animation/MotionTrajectory/Source/MotionTrajectory/Public/CharacterTrajectoryComponent.h
--- a/Engine/Plugins/Experimental/Animation/MotionTrajectory/Source/MotionTrajectory/Public/CharacterTrajectoryComponent.h
+++ b/Engine/Plugins/Experimental/Animation/MotionTrajectory/Source/MotionTrajectory/Public/CharacterTrajectoryComponent.h
@@ -105,6 +105,8 @@ protected:
 	float SecondsPerHistorySample = 0.f;
 	float SecondsPerPredictionSample = 0.f;
 
+	uint32 LastUpdateFrameNumber = 0;
+
 	// Forward axis for the SkeletalMeshComponent. It's common for skeletal mesh and animation data to not be X forward.
 	FQuat ForwardFacingCS = FQuat::Identity;
 
