Remove ensure for DeltaTime == 0. This can happen when TickAnimation(0.f) is called, e.g. from RefreshBoneTransforms, which isn't an error condition.

diff --git a/Engine/Plugins/Experimental/Animation/MotionTrajectory/Source/MotionTrajectory/Private/CharacterTrajectoryComponent.cpp b/Engine/Plugins/Experimental/Animation/MotionTrajectory/Source/MotionTrajectory/Private/CharacterTrajectoryComponent.cpp
--- a/Engine/Plugins/Experimental/Animation/MotionTrajectory/Source/MotionTrajectory/Private/CharacterTrajectoryComponent.cpp
+++ b/Engine/Plugins/Experimental/Animation/MotionTrajectory/Source/MotionTrajectory/Private/CharacterTrajectoryComponent.cpp
@@ -138,7 +138,7 @@ void UCharacterTrajectoryComponent::UpdateTrajectory(float DeltaSeconds)
 		return;
 	}
 
-	if (!ensure(DeltaSeconds > 0.f))
+	if (DeltaSeconds <= 0.f)
 	{
 		return;
 	}
