Move trajectory logic to a function library

diff --git a/Engine/Plugins/Experimental/Animation/MotionTrajectory/Source/MotionTrajectory/Private/MotionTrajectoryLibrary.cpp b/Engine/Plugins/Experimental/Animation/MotionTrajectory/Source/MotionTrajectory/Private/MotionTrajectoryLibrary.cpp
new file mode 100644
--- /dev/null
+++ b/Engine/Plugins/Experimental/Animation/MotionTrajectory/Source/MotionTrajectory/Private/MotionTrajectoryLibrary.cpp
@@ -0,0 +1,221 @@
+// Copyright Epic Games, Inc. All Rights Reserved.
+
+#include "MotionTrajectoryLibrary.h"
+#include "CharacterMovementTrajectoryLibrary.h"
+#include "Components/SkeletalMeshComponent.h"
+#include "GameFramework/Character.h"
+#include "GameFramework/CharacterMovementComponent.h"
+
+void FTrajectorySamplingData::Init()
+{
+	// The UI clamps these to be non-zero.
+	check(HistorySamplesPerSecond);
+	check(PredictionSamplesPerSecond);
+
+	NumHistorySamples = FMath::CeilToInt32(HistoryLengthSeconds * HistorySamplesPerSecond);
+	SecondsPerHistorySample = 1.f / HistorySamplesPerSecond;
+
+	NumPredictionSamples = FMath::CeilToInt32(PredictionLengthSeconds * PredictionSamplesPerSecond);
+	SecondsPerPredictionSample = 1.f / PredictionSamplesPerSecond;
+}
+
+void FCharacterTrajectoryData::Init(const AActor* Actor)
+{
+	Character = Cast<ACharacter>(Actor);
+	if (!ensureMsgf(Character, TEXT("FCharacterTrajectoryData requires valid ACharacter owner.")))
+	{
+		return;
+	}
+
+	SkelMeshComponent = Character->GetMesh();
+	if (!ensureMsgf(SkelMeshComponent, TEXT("FCharacterTrajectoryData must be run on an ACharacter with a valid USkeletalMeshComponent.")))
+	{
+		return;
+	}
+
+	CharacterMovementComponent = Character->GetCharacterMovement();
+	if (!ensureMsgf(CharacterMovementComponent, TEXT("FCharacterTrajectoryData must be run on an ACharacter with a valid UCharacterMovementComponent.")))
+	{
+		return;
+	}
+}
+
+bool FCharacterTrajectoryData::IsValid() const
+{
+	// Init() won't initialize CharacterMovementComponent if Character or SkelMeshComponent are null.
+	check(!CharacterMovementComponent || (Character && SkelMeshComponent));
+
+	return (CharacterMovementComponent != nullptr);
+}
+
+void FMotionTrajectoryLibrary::InitTrajectorySamples(FPoseSearchQueryTrajectory& Trajectory,
+	const FCharacterTrajectoryData& CharacterTrajectoryData, const FTrajectorySamplingData& SamplingData)
+{
+	if (!ensure(CharacterTrajectoryData.SkelMeshComponent))
+	{
+		return;
+	}
+
+	const FVector PositionWS = CharacterTrajectoryData.SkelMeshComponent->GetComponentLocation();;
+	const FQuat FacingWS = CharacterTrajectoryData.SkelMeshComponent->GetComponentRotation().Quaternion();
+
+	// History + current sample + prediction
+	Trajectory.Samples.SetNumUninitialized(SamplingData.NumHistorySamples + 1 + SamplingData.NumPredictionSamples);
+
+	// Initialize history samples
+	for (int32 i = 0; i < SamplingData.NumHistorySamples; ++i)
+	{
+		Trajectory.Samples[i].Position = PositionWS;
+		Trajectory.Samples[i].Facing = FacingWS;
+		Trajectory.Samples[i].AccumulatedSeconds = SamplingData.SecondsPerHistorySample * (i - SamplingData.NumHistorySamples);
+	}
+
+	// Initialize current sample and prediction
+	for (int32 i = SamplingData.NumHistorySamples; i < Trajectory.Samples.Num(); ++i)
+	{
+		Trajectory.Samples[i].Position = PositionWS;
+		Trajectory.Samples[i].Facing = FacingWS;
+		Trajectory.Samples[i].AccumulatedSeconds = SamplingData.SecondsPerPredictionSample * (i - SamplingData.NumHistorySamples);
+	}
+}
+
+void FMotionTrajectoryLibrary::UpdateHistory_ShiftInWorldSpace(FPoseSearchQueryTrajectory& Trajectory,
+	const FTrajectorySamplingData& SamplingData, float DeltaSeconds)
+{
+	check(SamplingData.NumHistorySamples <= Trajectory.Samples.Num());
+
+	// Shift history Samples when it's time to record a new one.
+	if (SamplingData.NumHistorySamples > 0 && FMath::Abs(Trajectory.Samples[SamplingData.NumHistorySamples - 1].AccumulatedSeconds) >= SamplingData.SecondsPerHistorySample)
+	{
+		for (int32 Index = 0; Index < SamplingData.NumHistorySamples; ++Index)
+		{
+			Trajectory.Samples[Index] = Trajectory.Samples[Index + 1];
+			Trajectory.Samples[Index].AccumulatedSeconds -= DeltaSeconds;
+		}
+	}
+	else
+	{
+		for (int32 Index = 0; Index < SamplingData.NumHistorySamples; ++Index)
+		{
+			Trajectory.Samples[Index].AccumulatedSeconds -= DeltaSeconds;
+		}
+	}
+}
+
+FVector FMotionTrajectoryLibrary::RemapVectorMagnitudeWithCurve(const FVector& Vector, bool bUseCurve, const FRuntimeFloatCurve& Curve)
+{
+	if (bUseCurve)
+	{
+		const float Length = Vector.Length();
+		if (Length > UE_KINDA_SMALL_NUMBER)
+		{
+			const float RemappedLength = Curve.GetRichCurveConst()->Eval(Length);
+			return Vector * (RemappedLength / Length);
+		}
+	}
+
+	return Vector;
+}
+
+void FMotionTrajectoryLibrary::UpdatePrediction_SimulateCharacterMovement(FPoseSearchQueryTrajectory& Trajectory, FCharacterTrajectoryData& CharacterTrajectoryData,
+	const FTrajectorySamplingData& SamplingData, float DeltaSeconds)
+{
+	if (!ensure(CharacterTrajectoryData.IsValid()))
+	{
+		return;
+	}
+
+	FVector CurrentPositionWS = CharacterTrajectoryData.SkelMeshComponent->GetComponentLocation();
+	FVector CurrentVelocityWS = FMotionTrajectoryLibrary::RemapVectorMagnitudeWithCurve(CharacterTrajectoryData.CharacterMovementComponent->Velocity,
+		CharacterTrajectoryData.bUseSpeedRemappingCurve, CharacterTrajectoryData.SpeedRemappingCurve);
+	FVector CurrentAccelerationWS = FMotionTrajectoryLibrary::RemapVectorMagnitudeWithCurve(CharacterTrajectoryData.CharacterMovementComponent->GetCurrentAcceleration(),
+		CharacterTrajectoryData.bUseAccelerationRemappingCurve, CharacterTrajectoryData.AccelerationRemappingCurve);
+	FQuat CurrentFacingWS = CharacterTrajectoryData.SkelMeshComponent->GetComponentRotation().Quaternion();
+	FQuat SkelMeshCompRelativeRotation = CharacterTrajectoryData.SkelMeshComponent->GetRelativeRotation().Quaternion();
+
+	const FRotator ControllerRotationRate = FMotionTrajectoryLibrary::CalculateControllerRotationRate(DeltaSeconds, CharacterTrajectoryData);
+	FQuat ControllerRotationPerStep = (ControllerRotationRate * SamplingData.SecondsPerPredictionSample).Quaternion();
+
+	float AccumulatedSeconds = 0.f;
+
+	const int32 LastIndex = Trajectory.Samples.Num() - 1;
+	if (SamplingData.NumHistorySamples <= LastIndex)
+	{
+		for (int32 Index = SamplingData.NumHistorySamples; ; ++Index)
+		{
+			Trajectory.Samples[Index].Position = CurrentPositionWS;
+			Trajectory.Samples[Index].Facing = CurrentFacingWS;
+			Trajectory.Samples[Index].AccumulatedSeconds = AccumulatedSeconds;
+
+			if (Index == LastIndex)
+			{
+				break;
+			}
+
+			CurrentPositionWS += CurrentVelocityWS * SamplingData.SecondsPerPredictionSample;
+			AccumulatedSeconds += SamplingData.SecondsPerPredictionSample;
+
+			// Account for the controller (e.g. the camera) rotating.
+			CurrentFacingWS = ControllerRotationPerStep * CurrentFacingWS;
+			CurrentAccelerationWS = FMotionTrajectoryLibrary::RemapVectorMagnitudeWithCurve(ControllerRotationPerStep * CurrentAccelerationWS,
+				CharacterTrajectoryData.bUseAccelerationRemappingCurve, CharacterTrajectoryData.AccelerationRemappingCurve);
+
+			FVector NewVelocityCS = FVector::ZeroVector;
+			UCharacterMovementTrajectoryLibrary::StepCharacterMovementGroundPrediction(SamplingData.SecondsPerPredictionSample, CurrentVelocityWS, CurrentAccelerationWS, 
+				CharacterTrajectoryData.CharacterMovementComponent, NewVelocityCS);
+			CurrentVelocityWS = FMotionTrajectoryLibrary::RemapVectorMagnitudeWithCurve(NewVelocityCS,
+				CharacterTrajectoryData.bUseSpeedRemappingCurve, CharacterTrajectoryData.SpeedRemappingCurve);
+
+			if (CharacterTrajectoryData.CharacterMovementComponent->bOrientRotationToMovement && !CurrentAccelerationWS.IsNearlyZero())
+			{
+				// Rotate towards acceleration.
+				const FVector CurrentAccelerationCS = SkelMeshCompRelativeRotation.RotateVector(CurrentAccelerationWS);
+				CurrentFacingWS = FMath::QInterpConstantTo(CurrentFacingWS, CurrentAccelerationCS.ToOrientationQuat(), SamplingData.SecondsPerPredictionSample,
+					CharacterTrajectoryData.RotateTowardsMovementSpeed);
+			}
+		}
+	}
+}
+
+// Calculate how much the character is rotating each update due to the controller (e.g. the camera) rotating.
+// E.g. If the user is moving forward but rotating the camera, the character (and thus future accelerations, facing directions, etc) will rotate.
+FRotator FMotionTrajectoryLibrary::CalculateControllerRotationRate(float DeltaSeconds, FCharacterTrajectoryData& CharacterTrajectoryData)
+{
+	if (!ensure(DeltaSeconds > 0.f))
+	{
+		return FRotator::ZeroRotator;
+	}
+
+	if (!ensure(CharacterTrajectoryData.IsValid()))
+	{
+		return FRotator::ZeroRotator;
+	}
+
+	const AController* Controller = CharacterTrajectoryData.Character->Controller;
+	if (!Controller)
+	{
+		// @todo: Simulated proxies don't have controllers, so they'll need some other mechanism to account for controller rotation rate.
+		return FRotator::ZeroRotator;
+	}
+
+	FRotator DesiredControllerRotation = Controller->GetDesiredRotation();
+	if (CharacterTrajectoryData.CharacterMovementComponent->ShouldRemainVertical())
+	{
+		DesiredControllerRotation.Yaw = FRotator::NormalizeAxis(DesiredControllerRotation.Yaw);
+		DesiredControllerRotation.Pitch = 0.f;
+		DesiredControllerRotation.Roll = 0.f;
+	}
+
+	const FRotator DesiredRotationDelta = DesiredControllerRotation - CharacterTrajectoryData.DesiredControllerRotationLastUpdate;
+	CharacterTrajectoryData.DesiredControllerRotationLastUpdate = DesiredControllerRotation;
+
+	FRotator ControllerRotationRate = DesiredRotationDelta.GetNormalized() * (1.f / DeltaSeconds);
+	if (CharacterTrajectoryData.MaxControllerRotationRate >= 0.f)
+	{
+		ControllerRotationRate.Pitch = FMath::Sign(ControllerRotationRate.Pitch) * FMath::Min(FMath::Abs(ControllerRotationRate.Pitch), CharacterTrajectoryData.MaxControllerRotationRate);
+		ControllerRotationRate.Yaw = FMath::Sign(ControllerRotationRate.Yaw) * FMath::Min(FMath::Abs(ControllerRotationRate.Yaw), CharacterTrajectoryData.MaxControllerRotationRate);
+		ControllerRotationRate.Roll = FMath::Sign(ControllerRotationRate.Roll) * FMath::Min(FMath::Abs(ControllerRotationRate.Roll), CharacterTrajectoryData.MaxControllerRotationRate);
+	}
+
+	return ControllerRotationRate;
+}
\ No newline at end of file
diff --git a/Engine/Plugins/Experimental/Animation/MotionTrajectory/Source/MotionTrajectory/Public/MotionTrajectoryLibrary.h b/Engine/Plugins/Experimental/Animation/MotionTrajectory/Source/MotionTrajectory/Public/MotionTrajectoryLibrary.h
new file mode 100644
--- /dev/null
+++ b/Engine/Plugins/Experimental/Animation/MotionTrajectory/Source/MotionTrajectory/Public/MotionTrajectoryLibrary.h
@@ -0,0 +1,111 @@
+// Copyright Epic Games, Inc. All Rights Reserved.
+
+#pragma once
+
+#include "Curves/CurveFloat.h"
+#include "PoseSearch/PoseSearchTrajectoryTypes.h"
+#include "MotionTrajectoryLibrary.generated.h"
+
+class ACharacter;
+class UCharacterMovementComponent;
+
+USTRUCT(BlueprintType)
+struct MOTIONTRAJECTORY_API FTrajectorySamplingData
+{
+	GENERATED_BODY()
+
+	void Init();
+
+	// This should generally match the longest history required by a Motion Matching Database in the project.
+	// Motion Matching will use extrapolation to generate samples if the history doesn't contain enough samples.
+	UPROPERTY(EditDefaultsOnly, BlueprintReadOnly, Category = "Trajectory Settings")
+	float HistoryLengthSeconds = 1.5f;
+
+	// Higher values will cost more storage and processing time, but give higher accuracy.
+	UPROPERTY(EditDefaultsOnly, BlueprintReadOnly, Category = "Trajectory Settings", meta = (ClampMin = "1", ClampMax = "120"))
+	int32 HistorySamplesPerSecond = 5;
+
+	// This should match the longest trajectory prediction required by a Motion Matching Database in the project.
+	// Motion Matching will use extrapolation to generate samples if the prediction doesn't contain enough samples.
+	UPROPERTY(EditDefaultsOnly, BlueprintReadOnly, Category = "Trajectory Settings")
+	float PredictionLengthSeconds = 1.5f;
+
+	// Higher values will cost more storage and processing time, but give higher accuracy.
+	UPROPERTY(EditDefaultsOnly, BlueprintReadOnly, Category = "Trajectory Settings", meta = (ClampMin = "1", ClampMax = "120"))
+	int32 PredictionSamplesPerSecond = 5;
+
+	int32 NumHistorySamples = -1;
+	float SecondsPerHistorySample = 0.f;
+
+	int32 NumPredictionSamples = -1;
+	float SecondsPerPredictionSample = 0.f;
+};
+
+USTRUCT(BlueprintType)
+struct MOTIONTRAJECTORY_API FCharacterTrajectoryData
+{
+	GENERATED_BODY()
+
+	void Init(const AActor* Actor);
+
+	bool IsValid() const;
+
+	// If the character is forward facing (i.e. bOrientRotationToMovement is true), this controls how quickly the trajectory will rotate
+	// to face acceleration. It's common for this to differ from the rotation rate of the character, because animations are often authored 
+	// with different rotation speeds than the character. This is especially true in cases where the character rotation snaps to movement.
+	UPROPERTY(EditDefaultsOnly, BlueprintReadOnly, Category = "Trajectory Settings")
+	float RotateTowardsMovementSpeed = 10.f;
+
+	// Maximum controller rotation rate in degrees per second used to clamp the character owner controller desired rotation to generate the prediction trajectory.
+	// Negative values disable the clamping behavior
+	UPROPERTY(EditDefaultsOnly, Category = "Trajectory Settings")
+	float MaxControllerRotationRate = -1.f;
+
+	UPROPERTY(EditDefaultsOnly, Category = "Trajectory Settings", meta = (InlineEditConditionToggle))
+	bool bUseSpeedRemappingCurve = false;
+
+	UPROPERTY(EditDefaultsOnly, Category = "Trajectory Settings", meta = (EditCondition = "bUseSpeedRemappingCurve"))
+	FRuntimeFloatCurve SpeedRemappingCurve;
+
+	UPROPERTY(EditDefaultsOnly, Category = "Trajectory Settings", meta = (InlineEditConditionToggle))
+	bool bUseAccelerationRemappingCurve = false;
+
+	UPROPERTY(EditDefaultsOnly, Category = "Trajectory Settings", meta = (EditCondition = "bUseAccelerationRemappingCurve"))
+	FRuntimeFloatCurve AccelerationRemappingCurve;
+
+	UPROPERTY(Transient)
+	TObjectPtr<const ACharacter> Character = nullptr;
+
+	UPROPERTY(Transient)
+	TObjectPtr<const USkeletalMeshComponent> SkelMeshComponent = nullptr;
+
+	UPROPERTY(Transient)
+	TObjectPtr<const UCharacterMovementComponent> CharacterMovementComponent = nullptr;
+
+	FRotator DesiredControllerRotationLastUpdate = FRotator::ZeroRotator;
+};
+
+/**
+ * Set of functions to help populate a FPoseSearchQueryTrajectory for motion matching.
+ * UCharacterTrajectoryComponent uses these functions, but they can also be used by a UAnimInstance to avoid the component.
+ */
+struct MOTIONTRAJECTORY_API FMotionTrajectoryLibrary
+{
+public:
+	static void InitTrajectorySamples(FPoseSearchQueryTrajectory& Trajectory,
+		const FCharacterTrajectoryData& CharacterTrajectoryData, const FTrajectorySamplingData& SamplingData);
+
+	// Update history by shifting world space samples back in time. This gives us a history of world space transforms.
+	// This approach is simple, but it causes problems with moving platforms.
+	static void UpdateHistory_ShiftInWorldSpace(FPoseSearchQueryTrajectory& Trajectory,
+		const FTrajectorySamplingData& SamplingData, float DeltaSeconds);
+
+	// Update prediction by simulating the movement math for ground locomotion in UCharacterMovementComponent.
+	static void UpdatePrediction_SimulateCharacterMovement(FPoseSearchQueryTrajectory& Trajectory, FCharacterTrajectoryData& CharacterTrajectoryData,
+		const FTrajectorySamplingData& SamplingData, float DeltaSeconds);
+
+private:
+	static FVector RemapVectorMagnitudeWithCurve(const FVector& Vector, bool bUseCurve, const FRuntimeFloatCurve& Curve);
+
+	static FRotator CalculateControllerRotationRate(float DeltaSeconds, FCharacterTrajectoryData& CharacterTrajectoryData);
+};
\ No newline at end of file
