Update UCharacterTrajectoryComponent to use new library functions

diff --git a/Engine/Plugins/Experimental/Animation/MotionTrajectory/Source/MotionTrajectory/Private/CharacterTrajectoryComponent.cpp b/Engine/Plugins/Experimental/Animation/MotionTrajectory/Source/MotionTrajectory/Private/CharacterTrajectoryComponent.cpp
--- a/Engine/Plugins/Experimental/Animation/MotionTrajectory/Source/MotionTrajectory/Private/CharacterTrajectoryComponent.cpp
+++ b/Engine/Plugins/Experimental/Animation/MotionTrajectory/Source/MotionTrajectory/Private/CharacterTrajectoryComponent.cpp
@@ -2,10 +2,7 @@
 
 #include "CharacterTrajectoryComponent.h"
 
-#include "CharacterMovementTrajectoryLibrary.h"
-#include "Components/SkeletalMeshComponent.h"
 #include "GameFramework/Character.h"
-#include "GameFramework/CharacterMovementComponent.h"
 #include "HAL/IConsoleManager.h"
 #include "MotionTrajectory.h"
 
@@ -26,129 +23,55 @@ void UCharacterTrajectoryComponent::InitializeComponent()
 {
 	Super::InitializeComponent();
 
-	if (bAutoUpdateTrajectory)
+	if (ACharacter* Character = Cast<ACharacter>(GetOwner()))
 	{
-		if (ACharacter* Character = Cast<ACharacter>(GetOwner()))
-		{
-			Character->OnCharacterMovementUpdated.AddDynamic(this, &UCharacterTrajectoryComponent::OnMovementUpdated);
-		}
-		else
-		{
-			UE_LOG(LogMotionTrajectory, Error, TEXT("UCharacterTrajectoryComponent requires its owner to be ACharacter"));
-		}
+		Character->OnCharacterMovementUpdated.AddDynamic(this, &UCharacterTrajectoryComponent::OnMovementUpdated);
+	}
+	else
+	{
+		UE_LOG(LogMotionTrajectory, Error, TEXT("UCharacterTrajectoryComponent requires its owner to be ACharacter"));
 	}
 }
 
 void UCharacterTrajectoryComponent::UninitializeComponent()
 {
-	if (bAutoUpdateTrajectory)
+	if (ACharacter* Character = Cast<ACharacter>(GetOwner()))
+	{
+		Character->OnCharacterMovementUpdated.RemoveDynamic(this, &UCharacterTrajectoryComponent::OnMovementUpdated);
+	}
+	else
 	{
-		if (ACharacter* Character = Cast<ACharacter>(GetOwner()))
-		{
-			Character->OnCharacterMovementUpdated.RemoveDynamic(this, &UCharacterTrajectoryComponent::OnMovementUpdated);
-		}
-		else
-		{
-			UE_LOG(LogMotionTrajectory, Error, TEXT("UCharacterTrajectoryComponent requires its owner to be ACharacter"));
-		}
+		UE_LOG(LogMotionTrajectory, Error, TEXT("UCharacterTrajectoryComponent requires its owner to be ACharacter"));
 	}
 
 	Super::UninitializeComponent();
 }
 
-FRotator UCharacterTrajectoryComponent::GetFacingFromMeshComponent_Implementation(const USkeletalMeshComponent* MeshComponent) const
-{
-	check(MeshComponent);
-	return MeshComponent->GetComponentRotation();
-}
-
 void UCharacterTrajectoryComponent::BeginPlay()
 {
 	Super::BeginPlay();
 
-	const ACharacter* Character = Cast<ACharacter>(GetOwner());
-	if (!ensureMsgf(Character, TEXT("UCharacterTrajectoryComponent requires valid ACharacter owner.")))
-	{
-		return;
-	}
-
-	SkelMeshComponent = Character->GetMesh();
-	if (!ensureMsgf(SkelMeshComponent, TEXT("UCharacterTrajectoryComponent must be run on an ACharacter with a valid USkeletalMeshComponent.")))
-	{
-		return;
-	}
-
-	CharacterMovementComponent = Character->GetCharacterMovement();
-	if (!ensureMsgf(CharacterMovementComponent, TEXT("UCharacterTrajectoryComponent must be run on an ACharacter with a valid UCharacterMovementComponent.")))
-	{
-		return;
-	}
-
-	// Default forward in the engine is the X axis, but data often diverges from this (e.g. it's common for skeletal meshes to be Y forward).
-	// We determine the forward direction in the space of the skeletal mesh component based on the offset from the actor.
-	ForwardFacingCS = SkelMeshComponent->GetRelativeRotation().Quaternion().Inverse();
-
-	// The UI clamps these to be non-zero.
-	check(HistorySamplesPerSecond);
-	check(PredictionSamplesPerSecond);
+	SamplingData.Init();
+	CharacterTrajectoryData.Init(GetOwner());
 
-	NumHistorySamples = FMath::CeilToInt32(HistoryLengthSeconds * HistorySamplesPerSecond);
-	SecondsPerHistorySample = 1.f / HistorySamplesPerSecond;
-
-	const int32 NumPredictionSamples = FMath::CeilToInt32(PredictionLengthSeconds * PredictionSamplesPerSecond);
-	SecondsPerPredictionSample = 1.f / PredictionSamplesPerSecond;
-
-	FPoseSearchQueryTrajectorySample DefaultSample;
-	DefaultSample.Facing = ForwardFacingCS;
-	DefaultSample.Position = FVector::ZeroVector;
-	DefaultSample.AccumulatedSeconds = 0.f;
-
-	const FVector SkelMeshComponentLocationWS = SkelMeshComponent->GetComponentLocation();
-	const FQuat FacingWS = FQuat(GetFacingFromMeshComponent(SkelMeshComponent));
-
-	// History + current sample + prediction
-	Trajectory.Samples.Init(DefaultSample, NumHistorySamples + 1 + NumPredictionSamples);
-
-	// initializing history samples AccumulatedSeconds
-	for (int32 i = 0; i < NumHistorySamples; ++i)
-	{
-		Trajectory.Samples[i].AccumulatedSeconds = SecondsPerHistorySample * (i - NumHistorySamples);
-		Trajectory.Samples[i].Position = SkelMeshComponentLocationWS;
-		Trajectory.Samples[i].Facing = FacingWS;
-	}
-
-	// initializing history samples AccumulatedSeconds
-	for (int32 i = NumHistorySamples + 1; i < Trajectory.Samples.Num(); ++i)
-	{
-		Trajectory.Samples[i].AccumulatedSeconds = SecondsPerPredictionSample * (i - NumHistorySamples);
-		Trajectory.Samples[i].Position = SkelMeshComponentLocationWS;
-		Trajectory.Samples[i].Facing = FacingWS;
-	}
+	FMotionTrajectoryLibrary::InitTrajectorySamples(Trajectory, CharacterTrajectoryData, SamplingData);
 }
 
-void UCharacterTrajectoryComponent::UpdateTrajectory(float DeltaSeconds)
+void UCharacterTrajectoryComponent::OnMovementUpdated(float DeltaSeconds, FVector OldLocation, FVector OldVelocity)
 {
 	if (LastUpdateFrameNumber != 0 && LastUpdateFrameNumber == GFrameNumber)
 	{
 		return;
 	}
 
-	if (!ensure(CharacterMovementComponent != nullptr && SkelMeshComponent != nullptr))
-	{
-		return;
-	}
-
 	if (DeltaSeconds <= 0.f)
 	{
 		return;
 	}
 
-	const FVector SkelMeshComponentLocationWS = SkelMeshComponent->GetComponentLocation();
-	const FQuat FacingWS = FQuat(GetFacingFromMeshComponent(SkelMeshComponent));
-	const FRotator ControllerRotationRate = CalculateControllerRotationRate(DeltaSeconds, CharacterMovementComponent->ShouldRemainVertical());
-		
-	UpdateHistory(DeltaSeconds);
-	UpdatePrediction(SkelMeshComponentLocationWS, FacingWS, CharacterMovementComponent->Velocity, CharacterMovementComponent->GetCurrentAcceleration(), ControllerRotationRate);
+	FMotionTrajectoryLibrary::UpdateHistory_ShiftInWorldSpace(Trajectory, SamplingData, DeltaSeconds);
+
+	FMotionTrajectoryLibrary::UpdatePrediction_SimulateCharacterMovement(Trajectory, CharacterTrajectoryData, SamplingData, DeltaSeconds);
 
 	LastUpdateFrameNumber = GFrameNumber;
 
@@ -158,133 +81,4 @@ void UCharacterTrajectoryComponent::UpdateTrajectory(float DeltaSeconds)
 		Trajectory.DebugDrawTrajectory(GetWorld());
 	}
 #endif // ENABLE_ANIM_DEBUG
-}
-
-void UCharacterTrajectoryComponent::OnMovementUpdated(float DeltaSeconds, FVector OldLocation, FVector OldVelocity)
-{
-	UpdateTrajectory(DeltaSeconds);
-}
-
-// This function shifts the range of history samples whenever a new history sample should be recorded.
-// This allows us to keep a single sample array in world space that can be read directly by the Motion Matching node. 
-void UCharacterTrajectoryComponent::UpdateHistory(float DeltaSeconds)
-{
-	check(NumHistorySamples <= Trajectory.Samples.Num());
-
-	// Shift history Samples when it's time to record a new one.
-	if (NumHistorySamples > 0 && FMath::Abs(Trajectory.Samples[NumHistorySamples - 1].AccumulatedSeconds) >= SecondsPerHistorySample)
-	{
-		for (int32 Index = 0; Index < NumHistorySamples; ++Index)
-		{
-			Trajectory.Samples[Index] = Trajectory.Samples[Index + 1];
-			Trajectory.Samples[Index].AccumulatedSeconds -= DeltaSeconds;
-		}
-	}
-	else
-	{
-		for (int32 Index = 0; Index < NumHistorySamples; ++Index)
-		{
-			Trajectory.Samples[Index].AccumulatedSeconds -= DeltaSeconds;
-		}
-	}
-}
-
-static FVector RemapVectorMagnitudeWithCurve(const FVector& Vector, bool bUseCurve, const FRuntimeFloatCurve& Curve)
-{
-	if (bUseCurve)
-	{
-		const float Length = Vector.Length();
-		if (Length > UE_KINDA_SMALL_NUMBER)
-		{
-			const float RemappedLength = Curve.GetRichCurveConst()->Eval(Length);
-			return Vector * (RemappedLength / Length);
-		}
-	}
-
-	return Vector;
-}
-
-void UCharacterTrajectoryComponent::UpdatePrediction(const FVector& PositionWS, const FQuat& FacingWS, const FVector& VelocityWS, const FVector& AccelerationWS, const FRotator& ControllerRotationRate)
-{
-	check(CharacterMovementComponent);
-
-	FVector CurrentPositionWS = PositionWS;
-	FVector CurrentVelocityWS = RemapVectorMagnitudeWithCurve(VelocityWS, bUseSpeedRemappingCurve, SpeedRemappingCurve);
-	FVector CurrentAccelerationWS = RemapVectorMagnitudeWithCurve(AccelerationWS, bUseAccelerationRemappingCurve, AccelerationRemappingCurve);
-	FQuat CurrentFacingWS = FacingWS;
-	float AccumulatedSeconds = 0.f;
-
-	FQuat ControllerRotationPerStep = (ControllerRotationRate * SecondsPerPredictionSample).Quaternion();
-
-	const int32 LastIndex = Trajectory.Samples.Num() - 1;
-	if (NumHistorySamples <= LastIndex)
-	{
-		for (int32 Index = NumHistorySamples; ; ++Index)
-		{
-			Trajectory.Samples[Index].Position = CurrentPositionWS;
-			Trajectory.Samples[Index].Facing = CurrentFacingWS;
-			Trajectory.Samples[Index].AccumulatedSeconds = AccumulatedSeconds;
-
-			if (Index == LastIndex)
-			{
-				break;
-			}
-
-			CurrentPositionWS += CurrentVelocityWS * SecondsPerPredictionSample;
-			AccumulatedSeconds += SecondsPerPredictionSample;
-
-			// Account for the controller (e.g. the camera) rotating.
-			CurrentFacingWS = ControllerRotationPerStep * CurrentFacingWS;
-			CurrentAccelerationWS = RemapVectorMagnitudeWithCurve(ControllerRotationPerStep * CurrentAccelerationWS, bUseAccelerationRemappingCurve, AccelerationRemappingCurve);
-
-			FVector NewVelocityCS = FVector::ZeroVector;
-			UCharacterMovementTrajectoryLibrary::StepCharacterMovementGroundPrediction(SecondsPerPredictionSample, CurrentVelocityWS, CurrentAccelerationWS, CharacterMovementComponent, NewVelocityCS);
-			CurrentVelocityWS = RemapVectorMagnitudeWithCurve(NewVelocityCS, bUseSpeedRemappingCurve, SpeedRemappingCurve);
-
-			if (CharacterMovementComponent->bOrientRotationToMovement && !CurrentAccelerationWS.IsNearlyZero())
-			{
-				// Rotate towards acceleration.
-				const FVector CurrentAccelerationCS = SkelMeshComponent->GetRelativeRotation().Quaternion().RotateVector(CurrentAccelerationWS);
-				CurrentFacingWS = FMath::QInterpConstantTo(CurrentFacingWS, CurrentAccelerationCS.ToOrientationQuat(), SecondsPerPredictionSample, RotateTowardsMovementSpeed);
-			}
-		}
-	}
-}
-
-// Calculate how much the character is rotating each update due to the controller (e.g. the camera) rotating.
-// E.g. If the user is moving forward but rotating the camera, the character (and thus future accelerations, facing directions, etc) will rotate.
-FRotator UCharacterTrajectoryComponent::CalculateControllerRotationRate(float DeltaSeconds, bool bShouldRemainVertical)
-{
-	check(GetOwner());
-
-	// UpdateTrajectory handles DeltaSeconds == 0.f, so we should never hit this.
-	check(DeltaSeconds > 0.f);
-
-	ACharacter* CharacterOwner = Cast<ACharacter>(GetOwner());
-	if (CharacterOwner == nullptr || CharacterOwner->Controller == nullptr)
-	{
-		// @todo: Simulated proxies don't have controllers, so they'll need some other mechanism to account for controller rotation rate.
-		return FRotator::ZeroRotator;
-	}
-
-	FRotator DesiredControllerRotation = CharacterOwner->Controller->GetDesiredRotation();
-	if (bShouldRemainVertical)
-	{
-		DesiredControllerRotation.Yaw = FRotator::NormalizeAxis(DesiredControllerRotation.Yaw);
-		DesiredControllerRotation.Pitch = 0.f;
-		DesiredControllerRotation.Roll = 0.f;
-	}
-
-	const FRotator DesiredRotationDelta = DesiredControllerRotation - DesiredControllerRotationLastUpdate;
-	DesiredControllerRotationLastUpdate = DesiredControllerRotation;
-
-	FRotator ControllerRotationRate = DesiredRotationDelta.GetNormalized() * (1.f / DeltaSeconds);
-	if (MaxControllerRotationRate >= 0.f)
-	{
-		ControllerRotationRate.Pitch = FMath::Sign(ControllerRotationRate.Pitch) * FMath::Min(FMath::Abs(ControllerRotationRate.Pitch), MaxControllerRotationRate);
-		ControllerRotationRate.Yaw = FMath::Sign(ControllerRotationRate.Yaw) * FMath::Min(FMath::Abs(ControllerRotationRate.Yaw), MaxControllerRotationRate);
-		ControllerRotationRate.Roll = FMath::Sign(ControllerRotationRate.Roll) * FMath::Min(FMath::Abs(ControllerRotationRate.Roll), MaxControllerRotationRate);
-	}
-
-	return ControllerRotationRate;
 }
\ No newline at end of file
diff --git a/Engine/Plugins/Experimental/Animation/MotionTrajectory/Source/MotionTrajectory/Public/CharacterTrajectoryComponent.h b/Engine/Plugins/Experimental/Animation/MotionTrajectory/Source/MotionTrajectory/Public/CharacterTrajectoryComponent.h
--- a/Engine/Plugins/Experimental/Animation/MotionTrajectory/Source/MotionTrajectory/Public/CharacterTrajectoryComponent.h
+++ b/Engine/Plugins/Experimental/Animation/MotionTrajectory/Source/MotionTrajectory/Public/CharacterTrajectoryComponent.h
@@ -3,12 +3,10 @@
 #pragma once
 
 #include "Components/ActorComponent.h"
-#include "Curves/CurveFloat.h"
+#include "MotionTrajectoryLibrary.h"
 #include "PoseSearch/PoseSearchTrajectoryTypes.h"
 #include "CharacterTrajectoryComponent.generated.h"
 
-class UCharacterMovementComponent;
-
 // Component for generating trajectories usable by Motion Matching. This component generates trajectories from ACharacter.
 // This is intended to provide an example and starting point for using Motion Matching with a common setup using the default UCharacterMovementComponent.
 // It is expected work flow to extend or replace this component for projects that use a custom movement component or custom movement modes.
@@ -27,88 +25,19 @@ public:
 	virtual void BeginPlay() override;
 	// End UActorComponent Interface
 
-	// This can be called when bAutoUpdateTrajectory is false to manually control when trajectory updates.
-	UFUNCTION(BlueprintCallable, Category = "Trajectory")
-	void UpdateTrajectory(float DeltaSeconds);
-
-	// This can be used to override the facing used for trajectory calculation.
-	UFUNCTION(BlueprintNativeEvent, Category = "Trajectory")
-	FRotator GetFacingFromMeshComponent(const USkeletalMeshComponent* MeshComponent) const;
-	virtual FRotator GetFacingFromMeshComponent_Implementation(const USkeletalMeshComponent* MeshComponent) const;
-
 protected:
 	UFUNCTION()
 	void OnMovementUpdated(float DeltaSeconds, FVector OldLocation, FVector OldVelocity);
 
-	void UpdateHistory(float DeltaSeconds);
-	void UpdatePrediction(const FVector& PositionWS, const FQuat& FacingWS, const FVector& VelocityWS, const FVector& AccelerationWS, const FRotator& ControllerRotationRate);
-
-	FRotator CalculateControllerRotationRate(float DeltaSeconds, bool bShouldRemainVertical);
-
 	// Trajectory stored in world space so it can be directly passed to Motion Matching.
 	UPROPERTY(Transient, BlueprintReadOnly, Category = "Trajectory")
 	FPoseSearchQueryTrajectory Trajectory;
 
-	// This should generally match the longest history required by a Motion Matching Database in the project.
-	// Motion Matching will use extrapolation to generate samples if the history doesn't contain enough samples.
-	UPROPERTY(EditDefaultsOnly, BlueprintReadOnly, Category = "Trajectory Settings")
-	float HistoryLengthSeconds = 1.5f;
-
-	// Higher values will cost more storage and processing time, but give higher accuracy.
-	UPROPERTY(EditDefaultsOnly, BlueprintReadOnly, Category = "Trajectory Settings", meta = (ClampMin = "1", ClampMax = "120"))
-	int32 HistorySamplesPerSecond = 5;
-
-	// This should match the longest trajectory prediction required by a Motion Matching Database in the project.
-	// Motion Matching will use extrapolation to generate samples if the prediction doesn't contain enough samples.
-	UPROPERTY(EditDefaultsOnly, BlueprintReadOnly, Category = "Trajectory Settings")
-	float PredictionLengthSeconds = 1.5f;
-
-	// Higher values will cost more storage and processing time, but give higher accuracy.
-	UPROPERTY(EditDefaultsOnly, BlueprintReadOnly, Category = "Trajectory Settings", meta = (ClampMin = "1", ClampMax = "120"))
-	int32 PredictionSamplesPerSecond = 5;
-
-	// If the character is forward facing (i.e. bOrientRotationToMovement is true), this controls how quickly the trajectory will rotate
-	// to face acceleration. It's common for this to differ from the rotation rate of the character, because animations are often authored 
-	// with different rotation speeds than the character. This is especially true in cases where the character rotation snaps to movement.
-	UPROPERTY(EditDefaultsOnly, BlueprintReadOnly, Category = "Trajectory Settings")
-	float RotateTowardsMovementSpeed = 10.f;
-
-	// By default the component will always update trajectory. If desired, this can be disabled and the game can choose when to update.
-	// For example, a game might want to only update trajectory for characters that are within view or very close to the local player.
-	UPROPERTY(EditDefaultsOnly, BlueprintReadOnly, Category = "Trajectory Settings")
-	bool bAutoUpdateTrajectory = true;
+	UPROPERTY(Transient, BlueprintReadOnly, Category = "Trajectory Settings")
+	FTrajectorySamplingData SamplingData;
 
-	// Maximum controller rotation rate in degrees per second used to clamp the character owner controller desired rotation to generate the prediction trajectory.
-	// Negative values disable the clamping behavior
-	UPROPERTY(EditDefaultsOnly, Category = "Trajectory Settings")
-	float MaxControllerRotationRate = -1.f;	
-
-	UPROPERTY(EditDefaultsOnly, Category = "Trajectory Settings", meta = (InlineEditConditionToggle))
-	bool bUseSpeedRemappingCurve = false;
-
-	UPROPERTY(EditDefaultsOnly, Category = "Trajectory Settings", meta = (EditCondition = "bUseSpeedRemappingCurve"))
-	FRuntimeFloatCurve SpeedRemappingCurve;
-
-	UPROPERTY(EditDefaultsOnly, Category = "Trajectory Settings", meta = (InlineEditConditionToggle))
-	bool bUseAccelerationRemappingCurve = false;
-
-	UPROPERTY(EditDefaultsOnly, Category = "Trajectory Settings", meta = (EditCondition = "bUseAccelerationRemappingCurve"))
-	FRuntimeFloatCurve AccelerationRemappingCurve;
-
-	UPROPERTY()
-	TObjectPtr<USkeletalMeshComponent> SkelMeshComponent;
-
-	UPROPERTY()
-	TObjectPtr<UCharacterMovementComponent> CharacterMovementComponent;
-
-	int32 NumHistorySamples = -1;
-	float SecondsPerHistorySample = 0.f;
-	float SecondsPerPredictionSample = 0.f;
+	UPROPERTY(Transient, BlueprintReadOnly, Category = "Trajectory Settings")
+	FCharacterTrajectoryData CharacterTrajectoryData;
 
 	uint32 LastUpdateFrameNumber = 0;
-
-	// Forward axis for the SkeletalMeshComponent. It's common for skeletal mesh and animation data to not be X forward.
-	FQuat ForwardFacingCS = FQuat::Identity;
-
-	FRotator DesiredControllerRotationLastUpdate = FRotator::ZeroRotator;
 };
\ No newline at end of file
