Option to calculate trajectory by tracking a history of translations caused by CMC velocity (i.e. movement intent) and applying that to the current world position. This works well for moving platforms and moving on low friction surfaces. Simply tracking world transform history causes bad matching because there's motion that didn't come from movement intent.

diff --git a/Engine/Plugins/Experimental/Animation/MotionTrajectory/Source/MotionTrajectory/Private/CharacterTrajectoryComponent.cpp b/Engine/Plugins/Experimental/Animation/MotionTrajectory/Source/MotionTrajectory/Private/CharacterTrajectoryComponent.cpp
--- a/Engine/Plugins/Experimental/Animation/MotionTrajectory/Source/MotionTrajectory/Private/CharacterTrajectoryComponent.cpp
+++ b/Engine/Plugins/Experimental/Animation/MotionTrajectory/Source/MotionTrajectory/Private/CharacterTrajectoryComponent.cpp
@@ -69,6 +69,8 @@ void UCharacterTrajectoryComponent::OnMovementUpdated(float DeltaSeconds, FVecto
 		return;
 	}
 
+	CharacterTrajectoryData.Update(DeltaSeconds);
+
 	FMotionTrajectoryLibrary::UpdateHistory_ShiftInWorldSpace(Trajectory, SamplingData, DeltaSeconds);
 
 	FMotionTrajectoryLibrary::UpdatePrediction_SimulateCharacterMovement(Trajectory, CharacterTrajectoryData, SamplingData, DeltaSeconds);
diff --git a/Engine/Plugins/Experimental/Animation/MotionTrajectory/Source/MotionTrajectory/Private/MotionTrajectoryLibrary.cpp b/Engine/Plugins/Experimental/Animation/MotionTrajectory/Source/MotionTrajectory/Private/MotionTrajectoryLibrary.cpp
--- a/Engine/Plugins/Experimental/Animation/MotionTrajectory/Source/MotionTrajectory/Private/MotionTrajectoryLibrary.cpp
+++ b/Engine/Plugins/Experimental/Animation/MotionTrajectory/Source/MotionTrajectory/Private/MotionTrajectoryLibrary.cpp
@@ -40,6 +40,11 @@ void FCharacterTrajectoryData::Init(const AActor* Actor)
 	}
 }
 
+void FCharacterTrajectoryData::Update(float DeltaSeconds)
+{
+	UpdateControllerRotationRate(DeltaSeconds);
+}
+
 bool FCharacterTrajectoryData::IsValid() const
 {
 	// Init() won't initialize CharacterMovementComponent if Character or SkelMeshComponent are null.
@@ -48,6 +53,47 @@ bool FCharacterTrajectoryData::IsValid() const
 	return (CharacterMovementComponent != nullptr);
 }
 
+void FCharacterTrajectoryData::UpdateControllerRotationRate(float DeltaSeconds)
+{
+	ControllerRotationRate = FRotator::ZeroRotator;
+
+	if (!ensure(DeltaSeconds > 0.f))
+	{
+		return;
+	}
+
+	if (!ensure(IsValid()))
+	{
+		return;
+	}
+
+	const AController* Controller = Character->Controller;
+	if (!Controller)
+	{
+		// @todo: Simulated proxies don't have controllers, so they'll need some other mechanism to account for controller rotation rate.
+		return;
+	}
+
+	FRotator DesiredControllerRotation = Controller->GetDesiredRotation();
+	if (CharacterMovementComponent->ShouldRemainVertical())
+	{
+		DesiredControllerRotation.Yaw = FRotator::NormalizeAxis(DesiredControllerRotation.Yaw);
+		DesiredControllerRotation.Pitch = 0.f;
+		DesiredControllerRotation.Roll = 0.f;
+	}
+
+	const FRotator DesiredRotationDelta = DesiredControllerRotation - DesiredControllerRotationLastUpdate;
+	DesiredControllerRotationLastUpdate = DesiredControllerRotation;
+
+	ControllerRotationRate = DesiredRotationDelta.GetNormalized() * (1.f / DeltaSeconds);
+	if (MaxControllerRotationRate >= 0.f)
+	{
+		ControllerRotationRate.Pitch = FMath::Sign(ControllerRotationRate.Pitch) * FMath::Min(FMath::Abs(ControllerRotationRate.Pitch), MaxControllerRotationRate);
+		ControllerRotationRate.Yaw = FMath::Sign(ControllerRotationRate.Yaw) * FMath::Min(FMath::Abs(ControllerRotationRate.Yaw), MaxControllerRotationRate);
+		ControllerRotationRate.Roll = FMath::Sign(ControllerRotationRate.Roll) * FMath::Min(FMath::Abs(ControllerRotationRate.Roll), MaxControllerRotationRate);
+	}
+}
+
 void FMotionTrajectoryLibrary::InitTrajectorySamples(FPoseSearchQueryTrajectory& Trajectory,
 	const FCharacterTrajectoryData& CharacterTrajectoryData, const FTrajectorySamplingData& SamplingData)
 {
@@ -102,6 +148,48 @@ void FMotionTrajectoryLibrary::UpdateHistory_ShiftInWorldSpace(FPoseSearchQueryT
 	}
 }
 
+void FMotionTrajectoryLibrary::UpdateHistory_TransformHistory(FPoseSearchQueryTrajectory& Trajectory, TArrayView<FVector> TranslationHistory,
+	const FCharacterTrajectoryData& CharacterTrajectoryData, const FTrajectorySamplingData& SamplingData, float DeltaSeconds)
+{
+	check(SamplingData.NumHistorySamples <= Trajectory.Samples.Num());
+	check(TranslationHistory.Num() == SamplingData.NumHistorySamples);
+
+	const FVector& CurrentVelocityWS = CharacterTrajectoryData.CharacterMovementComponent->Velocity;
+	FVector CurrentTranslation = CurrentVelocityWS * DeltaSeconds;
+
+	// Shift history Samples when it's time to record a new one.
+	if (SamplingData.NumHistorySamples > 0 && FMath::Abs(Trajectory.Samples[SamplingData.NumHistorySamples - 1].AccumulatedSeconds) >= SamplingData.SecondsPerHistorySample)
+	{
+		for (int32 Index = 0; Index < SamplingData.NumHistorySamples - 1; ++Index)
+		{
+			Trajectory.Samples[Index].AccumulatedSeconds = Trajectory.Samples[Index + 1].AccumulatedSeconds;
+			TranslationHistory[Index] = TranslationHistory[Index + 1];
+			TranslationHistory[Index] += CurrentTranslation;
+		}
+
+		Trajectory.Samples[SamplingData.NumHistorySamples - 1].AccumulatedSeconds = 0.f;
+		TranslationHistory[SamplingData.NumHistorySamples - 1] = CurrentTranslation;
+	}
+	else
+	{
+		for (int32 Index = 0; Index < SamplingData.NumHistorySamples; ++Index)
+		{
+			TranslationHistory[Index] += CurrentTranslation;
+		}
+	}
+
+	// Update trajectory samples by applying the tracked translations to the current world position.
+	FVector CurrentPositionWS = CharacterTrajectoryData.SkelMeshComponent->GetComponentLocation();
+	for (int32 Index = 0; Index < SamplingData.NumHistorySamples; ++Index)
+	{
+		Trajectory.Samples[Index].AccumulatedSeconds -= DeltaSeconds;
+		Trajectory.Samples[Index].Position = CurrentPositionWS - TranslationHistory[Index];
+
+		// @todo: Handle facing. We currently don't use facing for history in any of our content. We will need the rotation intent of the character.
+		Trajectory.Samples[Index].Facing = FQuat::Identity;
+	}
+}
+
 FVector FMotionTrajectoryLibrary::RemapVectorMagnitudeWithCurve(const FVector& Vector, bool bUseCurve, const FRuntimeFloatCurve& Curve)
 {
 	if (bUseCurve)
@@ -117,8 +205,8 @@ FVector FMotionTrajectoryLibrary::RemapVectorMagnitudeWithCurve(const FVector& V
 	return Vector;
 }
 
-void FMotionTrajectoryLibrary::UpdatePrediction_SimulateCharacterMovement(FPoseSearchQueryTrajectory& Trajectory, FCharacterTrajectoryData& CharacterTrajectoryData,
-	const FTrajectorySamplingData& SamplingData, float DeltaSeconds)
+void FMotionTrajectoryLibrary::UpdatePrediction_SimulateCharacterMovement(FPoseSearchQueryTrajectory& Trajectory,
+	const FCharacterTrajectoryData& CharacterTrajectoryData, const FTrajectorySamplingData& SamplingData, float DeltaSeconds)
 {
 	if (!ensure(CharacterTrajectoryData.IsValid()))
 	{
@@ -133,8 +221,7 @@ void FMotionTrajectoryLibrary::UpdatePrediction_SimulateCharacterMovement(FPoseS
 	FQuat CurrentFacingWS = CharacterTrajectoryData.SkelMeshComponent->GetComponentRotation().Quaternion();
 	FQuat SkelMeshCompRelativeRotation = CharacterTrajectoryData.SkelMeshComponent->GetRelativeRotation().Quaternion();
 
-	const FRotator ControllerRotationRate = FMotionTrajectoryLibrary::CalculateControllerRotationRate(DeltaSeconds, CharacterTrajectoryData);
-	FQuat ControllerRotationPerStep = (ControllerRotationRate * SamplingData.SecondsPerPredictionSample).Quaternion();
+	FQuat ControllerRotationPerStep = (CharacterTrajectoryData.ControllerRotationRate * SamplingData.SecondsPerPredictionSample).Quaternion();
 
 	float AccumulatedSeconds = 0.f;
 
@@ -175,47 +262,4 @@ void FMotionTrajectoryLibrary::UpdatePrediction_SimulateCharacterMovement(FPoseS
 			}
 		}
 	}
-}
-
-// Calculate how much the character is rotating each update due to the controller (e.g. the camera) rotating.
-// E.g. If the user is moving forward but rotating the camera, the character (and thus future accelerations, facing directions, etc) will rotate.
-FRotator FMotionTrajectoryLibrary::CalculateControllerRotationRate(float DeltaSeconds, FCharacterTrajectoryData& CharacterTrajectoryData)
-{
-	if (!ensure(DeltaSeconds > 0.f))
-	{
-		return FRotator::ZeroRotator;
-	}
-
-	if (!ensure(CharacterTrajectoryData.IsValid()))
-	{
-		return FRotator::ZeroRotator;
-	}
-
-	const AController* Controller = CharacterTrajectoryData.Character->Controller;
-	if (!Controller)
-	{
-		// @todo: Simulated proxies don't have controllers, so they'll need some other mechanism to account for controller rotation rate.
-		return FRotator::ZeroRotator;
-	}
-
-	FRotator DesiredControllerRotation = Controller->GetDesiredRotation();
-	if (CharacterTrajectoryData.CharacterMovementComponent->ShouldRemainVertical())
-	{
-		DesiredControllerRotation.Yaw = FRotator::NormalizeAxis(DesiredControllerRotation.Yaw);
-		DesiredControllerRotation.Pitch = 0.f;
-		DesiredControllerRotation.Roll = 0.f;
-	}
-
-	const FRotator DesiredRotationDelta = DesiredControllerRotation - CharacterTrajectoryData.DesiredControllerRotationLastUpdate;
-	CharacterTrajectoryData.DesiredControllerRotationLastUpdate = DesiredControllerRotation;
-
-	FRotator ControllerRotationRate = DesiredRotationDelta.GetNormalized() * (1.f / DeltaSeconds);
-	if (CharacterTrajectoryData.MaxControllerRotationRate >= 0.f)
-	{
-		ControllerRotationRate.Pitch = FMath::Sign(ControllerRotationRate.Pitch) * FMath::Min(FMath::Abs(ControllerRotationRate.Pitch), CharacterTrajectoryData.MaxControllerRotationRate);
-		ControllerRotationRate.Yaw = FMath::Sign(ControllerRotationRate.Yaw) * FMath::Min(FMath::Abs(ControllerRotationRate.Yaw), CharacterTrajectoryData.MaxControllerRotationRate);
-		ControllerRotationRate.Roll = FMath::Sign(ControllerRotationRate.Roll) * FMath::Min(FMath::Abs(ControllerRotationRate.Roll), CharacterTrajectoryData.MaxControllerRotationRate);
-	}
-
-	return ControllerRotationRate;
 }
\ No newline at end of file
diff --git a/Engine/Plugins/Experimental/Animation/MotionTrajectory/Source/MotionTrajectory/Public/MotionTrajectoryLibrary.h b/Engine/Plugins/Experimental/Animation/MotionTrajectory/Source/MotionTrajectory/Public/MotionTrajectoryLibrary.h
--- a/Engine/Plugins/Experimental/Animation/MotionTrajectory/Source/MotionTrajectory/Public/MotionTrajectoryLibrary.h
+++ b/Engine/Plugins/Experimental/Animation/MotionTrajectory/Source/MotionTrajectory/Public/MotionTrajectoryLibrary.h
@@ -44,16 +44,18 @@ struct MOTIONTRAJECTORY_API FTrajectorySamplingData
 USTRUCT(BlueprintType)
 struct MOTIONTRAJECTORY_API FCharacterTrajectoryData
 {
+public:
 	GENERATED_BODY()
 
 	void Init(const AActor* Actor);
+	void Update(float DeltaSeconds);
 
 	bool IsValid() const;
 
 	// If the character is forward facing (i.e. bOrientRotationToMovement is true), this controls how quickly the trajectory will rotate
 	// to face acceleration. It's common for this to differ from the rotation rate of the character, because animations are often authored 
 	// with different rotation speeds than the character. This is especially true in cases where the character rotation snaps to movement.
-	UPROPERTY(EditDefaultsOnly, BlueprintReadOnly, Category = "Trajectory Settings")
+	UPROPERTY(EditDefaultsOnly, Category = "Trajectory Settings")
 	float RotateTowardsMovementSpeed = 10.f;
 
 	// Maximum controller rotation rate in degrees per second used to clamp the character owner controller desired rotation to generate the prediction trajectory.
@@ -82,6 +84,11 @@ struct MOTIONTRAJECTORY_API FCharacterTrajectoryData
 	UPROPERTY(Transient)
 	TObjectPtr<const UCharacterMovementComponent> CharacterMovementComponent = nullptr;
 
+	FRotator ControllerRotationRate = FRotator::ZeroRotator;
+
+private:
+	void UpdateControllerRotationRate(float DeltaSeconds);
+
 	FRotator DesiredControllerRotationLastUpdate = FRotator::ZeroRotator;
 };
 
@@ -100,12 +107,16 @@ public:
 	static void UpdateHistory_ShiftInWorldSpace(FPoseSearchQueryTrajectory& Trajectory,
 		const FTrajectorySamplingData& SamplingData, float DeltaSeconds);
 
+	// Update history by tracking offsets that result from character intent (e.g. movement component velocity) and applying
+	// that to the current world transform. This works well on moving platforms as it only stores a history of movement
+	// that results from character intent.
+	static void UpdateHistory_TransformHistory(FPoseSearchQueryTrajectory& Trajectory, TArrayView<FVector> TranslationHistory,
+		const FCharacterTrajectoryData& CharacterTrajectoryData, const FTrajectorySamplingData& SamplingData, float DeltaSeconds);
+
 	// Update prediction by simulating the movement math for ground locomotion in UCharacterMovementComponent.
-	static void UpdatePrediction_SimulateCharacterMovement(FPoseSearchQueryTrajectory& Trajectory, FCharacterTrajectoryData& CharacterTrajectoryData,
-		const FTrajectorySamplingData& SamplingData, float DeltaSeconds);
+	static void UpdatePrediction_SimulateCharacterMovement(FPoseSearchQueryTrajectory& Trajectory,
+		const FCharacterTrajectoryData& CharacterTrajectoryData, const FTrajectorySamplingData& SamplingData, float DeltaSeconds);
 
 private:
 	static FVector RemapVectorMagnitudeWithCurve(const FVector& Vector, bool bUseCurve, const FRuntimeFloatCurve& Curve);
-
-	static FRotator CalculateControllerRotationRate(float DeltaSeconds, FCharacterTrajectoryData& CharacterTrajectoryData);
 };
\ No newline at end of file
diff --git a/Engine/Plugins/Experimental/Animation/PoseSearch/Source/Runtime/Private/PoseSearchTrajectoryTypes.cpp b/Engine/Plugins/Experimental/Animation/PoseSearch/Source/Runtime/Private/PoseSearchTrajectoryTypes.cpp
--- a/Engine/Plugins/Experimental/Animation/PoseSearch/Source/Runtime/Private/PoseSearchTrajectoryTypes.cpp
+++ b/Engine/Plugins/Experimental/Animation/PoseSearch/Source/Runtime/Private/PoseSearchTrajectoryTypes.cpp
@@ -60,22 +60,24 @@ FPoseSearchQueryTrajectorySample FPoseSearchQueryTrajectory::GetSampleAtTime(flo
 }
 
 #if ENABLE_ANIM_DEBUG
-void FPoseSearchQueryTrajectory::DebugDrawTrajectory(const UWorld* World) const
+void FPoseSearchQueryTrajectory::DebugDrawTrajectory(const UWorld* World, float HeightOffset) const
 {
+	FVector OffsetVector = FVector::UpVector * HeightOffset;
+
 	const int32 LastIndex = Samples.Num() - 1;
 	if (LastIndex >= 0)
 	{
 		for (int32 Index = 0; ; ++Index)
 		{
-			DrawDebugSphere(World, Samples[Index].Position, 2.f /*Radius*/, 4 /*Segments*/, FColor::Black);
-			DrawDebugCoordinateSystem(World, Samples[Index].Position, FRotator(Samples[Index].Facing), 12.f /*Scale*/);
+			DrawDebugSphere(World, Samples[Index].Position + OffsetVector, 2.f /*Radius*/, 4 /*Segments*/, FColor::Black);
+			DrawDebugCoordinateSystem(World, Samples[Index].Position + OffsetVector, FRotator(Samples[Index].Facing), 12.f /*Scale*/);
 
 			if (Index == LastIndex)
 			{
 				break;
 			}
 			
-			DrawDebugLine(World, Samples[Index].Position, Samples[Index + 1].Position, FColor::Black);
+			DrawDebugLine(World, Samples[Index].Position + OffsetVector, Samples[Index + 1].Position + OffsetVector, FColor::Black);
 		}
 	}
 }
diff --git a/Engine/Plugins/Experimental/Animation/PoseSearch/Source/Runtime/Public/PoseSearch/PoseSearchTrajectoryTypes.h b/Engine/Plugins/Experimental/Animation/PoseSearch/Source/Runtime/Public/PoseSearch/PoseSearchTrajectoryTypes.h
--- a/Engine/Plugins/Experimental/Animation/PoseSearch/Source/Runtime/Public/PoseSearch/PoseSearchTrajectoryTypes.h
+++ b/Engine/Plugins/Experimental/Animation/PoseSearch/Source/Runtime/Public/PoseSearch/PoseSearchTrajectoryTypes.h
@@ -39,7 +39,7 @@ struct POSESEARCH_API FPoseSearchQueryTrajectory
 	FPoseSearchQueryTrajectorySample GetSampleAtTime(float Time, bool bExtrapolate = true) const;
 	
 #if ENABLE_ANIM_DEBUG
-	void DebugDrawTrajectory(const UWorld* World) const;
+	void DebugDrawTrajectory(const UWorld* World, float HeightOffset = 0.f) const;
 #endif // ENABLE_ANIM_DEBUG
 };
 POSESEARCH_API FArchive& operator<<(FArchive& Ar, FPoseSearchQueryTrajectory& Trajectory);
