motion trajectory - fix for crash in motion traj when compiling ACP with PIE running

diff --git a/Engine/Plugins/Experimental/Animation/MotionTrajectory/Source/MotionTrajectory/Public/MotionTrajectoryLibrary.h b/Engine/Plugins/Experimental/Animation/MotionTrajectory/Source/MotionTrajectory/Public/MotionTrajectoryLibrary.h
--- a/Engine/Plugins/Experimental/Animation/MotionTrajectory/Source/MotionTrajectory/Public/MotionTrajectoryLibrary.h
+++ b/Engine/Plugins/Experimental/Animation/MotionTrajectory/Source/MotionTrajectory/Public/MotionTrajectoryLibrary.h
@@ -34,10 +34,10 @@ struct MOTIONTRAJECTORY_API FTrajectorySamplingData
 	UPROPERTY(EditDefaultsOnly, BlueprintReadOnly, Category = "Trajectory Settings", meta = (ClampMin = "1", ClampMax = "120"))
 	int32 PredictionSamplesPerSecond = 5;
 
-	int32 NumHistorySamples = -1;
+	int32 NumHistorySamples = 0;
 	float SecondsPerHistorySample = 0.f;
 
-	int32 NumPredictionSamples = -1;
+	int32 NumPredictionSamples = 0;
 	float SecondsPerPredictionSample = 0.f;
 };
 
