Expose both clamped and unclamped controller rotation rate.

diff --git a/Engine/Plugins/Experimental/Animation/MotionTrajectory/Source/MotionTrajectory/Private/MotionTrajectoryLibrary.cpp b/Engine/Plugins/Experimental/Animation/MotionTrajectory/Source/MotionTrajectory/Private/MotionTrajectoryLibrary.cpp
--- a/Engine/Plugins/Experimental/Animation/MotionTrajectory/Source/MotionTrajectory/Private/MotionTrajectoryLibrary.cpp
+++ b/Engine/Plugins/Experimental/Animation/MotionTrajectory/Source/MotionTrajectory/Private/MotionTrajectoryLibrary.cpp
@@ -56,6 +56,7 @@ bool FCharacterTrajectoryData::IsValid() const
 void FCharacterTrajectoryData::UpdateControllerRotationRate(float DeltaSeconds)
 {
 	ControllerRotationRate = FRotator::ZeroRotator;
+	ControllerRotationRateClamped = FRotator::ZeroRotator;
 
 	if (!ensure(DeltaSeconds > 0.f))
 	{
@@ -88,9 +89,13 @@ void FCharacterTrajectoryData::UpdateControllerRotationRate(float DeltaSeconds)
 	ControllerRotationRate = DesiredRotationDelta.GetNormalized() * (1.f / DeltaSeconds);
 	if (MaxControllerRotationRate >= 0.f)
 	{
-		ControllerRotationRate.Pitch = FMath::Sign(ControllerRotationRate.Pitch) * FMath::Min(FMath::Abs(ControllerRotationRate.Pitch), MaxControllerRotationRate);
-		ControllerRotationRate.Yaw = FMath::Sign(ControllerRotationRate.Yaw) * FMath::Min(FMath::Abs(ControllerRotationRate.Yaw), MaxControllerRotationRate);
-		ControllerRotationRate.Roll = FMath::Sign(ControllerRotationRate.Roll) * FMath::Min(FMath::Abs(ControllerRotationRate.Roll), MaxControllerRotationRate);
+		ControllerRotationRateClamped.Pitch = FMath::Sign(ControllerRotationRate.Pitch) * FMath::Min(FMath::Abs(ControllerRotationRate.Pitch), MaxControllerRotationRate);
+		ControllerRotationRateClamped.Yaw = FMath::Sign(ControllerRotationRate.Yaw) * FMath::Min(FMath::Abs(ControllerRotationRate.Yaw), MaxControllerRotationRate);
+		ControllerRotationRateClamped.Roll = FMath::Sign(ControllerRotationRate.Roll) * FMath::Min(FMath::Abs(ControllerRotationRate.Roll), MaxControllerRotationRate);
+	}
+	else
+	{
+		ControllerRotationRateClamped = ControllerRotationRate;
 	}
 }
 
@@ -221,7 +226,7 @@ void FMotionTrajectoryLibrary::UpdatePrediction_SimulateCharacterMovement(FPoseS
 	FQuat CurrentFacingWS = CharacterTrajectoryData.SkelMeshComponent->GetComponentRotation().Quaternion();
 	FQuat SkelMeshCompRelativeRotation = CharacterTrajectoryData.SkelMeshComponent->GetRelativeRotation().Quaternion();
 
-	FQuat ControllerRotationPerStep = (CharacterTrajectoryData.ControllerRotationRate * SamplingData.SecondsPerPredictionSample).Quaternion();
+	FQuat ControllerRotationPerStep = (CharacterTrajectoryData.ControllerRotationRateClamped * SamplingData.SecondsPerPredictionSample).Quaternion();
 
 	float AccumulatedSeconds = 0.f;
 
diff --git a/Engine/Plugins/Experimental/Animation/MotionTrajectory/Source/MotionTrajectory/Public/MotionTrajectoryLibrary.h b/Engine/Plugins/Experimental/Animation/MotionTrajectory/Source/MotionTrajectory/Public/MotionTrajectoryLibrary.h
--- a/Engine/Plugins/Experimental/Animation/MotionTrajectory/Source/MotionTrajectory/Public/MotionTrajectoryLibrary.h
+++ b/Engine/Plugins/Experimental/Animation/MotionTrajectory/Source/MotionTrajectory/Public/MotionTrajectoryLibrary.h
@@ -85,6 +85,7 @@ public:
 	TObjectPtr<const UCharacterMovementComponent> CharacterMovementComponent = nullptr;
 
 	FRotator ControllerRotationRate = FRotator::ZeroRotator;
+	FRotator ControllerRotationRateClamped = FRotator::ZeroRotator;
 
 private:
 	void UpdateControllerRotationRate(float DeltaSeconds);
