motion trajectory - artificially bend character velocity towards acceleration direction to compute trajectory prediction, to get sharper turns

diff --git a/Engine/Plugins/Experimental/Animation/MotionTrajectory/Source/MotionTrajectory/Private/MotionTrajectoryLibrary.cpp b/Engine/Plugins/Experimental/Animation/MotionTrajectory/Source/MotionTrajectory/Private/MotionTrajectoryLibrary.cpp
--- a/Engine/Plugins/Experimental/Animation/MotionTrajectory/Source/MotionTrajectory/Private/MotionTrajectoryLibrary.cpp
+++ b/Engine/Plugins/Experimental/Animation/MotionTrajectory/Source/MotionTrajectory/Private/MotionTrajectoryLibrary.cpp
@@ -223,6 +223,32 @@ void FMotionTrajectoryLibrary::UpdatePrediction_SimulateCharacterMovement(FPoseS
 		CharacterTrajectoryData.bUseSpeedRemappingCurve, CharacterTrajectoryData.SpeedRemappingCurve);
 	FVector CurrentAccelerationWS = FMotionTrajectoryLibrary::RemapVectorMagnitudeWithCurve(CharacterTrajectoryData.CharacterMovementComponent->GetCurrentAcceleration(),
 		CharacterTrajectoryData.bUseAccelerationRemappingCurve, CharacterTrajectoryData.AccelerationRemappingCurve);
+
+	// bending CurrentVelocityWS towards CurrentAccelerationWS
+	if (CharacterTrajectoryData.BendVelocityTowardsAcceleration > UE_KINDA_SMALL_NUMBER && !CurrentAccelerationWS.IsNearlyZero())
+	{
+		const float CurrentSpeed = CurrentVelocityWS.Length();
+		const FVector VelocityWSAlongAcceleration = CurrentAccelerationWS.GetUnsafeNormal()* CurrentSpeed;
+		if (CharacterTrajectoryData.BendVelocityTowardsAcceleration < 1.f - UE_KINDA_SMALL_NUMBER)
+		{
+			CurrentVelocityWS = FMath::Lerp(CurrentVelocityWS, VelocityWSAlongAcceleration, CharacterTrajectoryData.BendVelocityTowardsAcceleration);
+			
+			const float NewLength = CurrentVelocityWS.Length();
+			if (NewLength > UE_KINDA_SMALL_NUMBER)
+			{
+				CurrentVelocityWS *= CurrentSpeed / NewLength;
+			}
+			else
+			{
+				// @todo: consider setting the CurrentVelocityWS = VelocityWSAlongAcceleration if vel and acc are in opposite directions
+			}
+		}
+		else
+		{
+			CurrentVelocityWS = VelocityWSAlongAcceleration;
+		}
+	}
+
 	FQuat CurrentFacingWS = CharacterTrajectoryData.SkelMeshComponent->GetComponentRotation().Quaternion();
 	FQuat SkelMeshCompRelativeRotation = CharacterTrajectoryData.SkelMeshComponent->GetRelativeRotation().Quaternion();
 
diff --git a/Engine/Plugins/Experimental/Animation/MotionTrajectory/Source/MotionTrajectory/Public/MotionTrajectoryLibrary.h b/Engine/Plugins/Experimental/Animation/MotionTrajectory/Source/MotionTrajectory/Public/MotionTrajectoryLibrary.h
--- a/Engine/Plugins/Experimental/Animation/MotionTrajectory/Source/MotionTrajectory/Public/MotionTrajectoryLibrary.h
+++ b/Engine/Plugins/Experimental/Animation/MotionTrajectory/Source/MotionTrajectory/Public/MotionTrajectoryLibrary.h
@@ -63,6 +63,11 @@ public:
 	UPROPERTY(EditDefaultsOnly, Category = "Trajectory Settings")
 	float MaxControllerRotationRate = -1.f;
 
+	// artificially bend character velocity towards acceleration direction to compute trajectory prediction, to get sharper turns
+	// 0: character velocity is used with no alteration, 1: the acceleration direction is used as velocity direction
+	UPROPERTY(EditDefaultsOnly, Category = "Trajectory Settings", meta = (ClampMin = "0", ClampMax = "1"))
+	float BendVelocityTowardsAcceleration = 0.f;
+
 	UPROPERTY(EditDefaultsOnly, Category = "Trajectory Settings", meta = (InlineEditConditionToggle))
 	bool bUseSpeedRemappingCurve = false;
 
