Trajectory refactor:
* Copy component data from the Character in FCharacterTrajectoryData::UpdateDataFromCharacter(). Game specific logic can then override values afterwards as needed. For example, if a powerup drastically changes speed/friction values, those values can be overwritten on FCharacterTrajectoryData to more sensible values for animation selection (i.e. we favor good animation selection over being accurate to the motion model in special cases).
* Removed FMotionTrajectoryLibrary::UpdateHistory_ShiftInWorldSpace(). UpdateHistory_TransformHistory() gives the same results in normal cases (e.g. ground locomotion) and much better results in exotic cases (moving platforms, low friction, etc).
* Removed UCharacterMovementTrajectoryLibrary, since it was only one function. Moved the prediction logic into FCharacterTrajectoryData.

diff --git a/Engine/Plugins/Experimental/Animation/MotionTrajectory/Source/MotionTrajectory/Private/CharacterMovementTrajectoryLibrary.cpp b/Engine/Plugins/Experimental/Animation/MotionTrajectory/Source/MotionTrajectory/Private/CharacterMovementTrajectoryLibrary.cpp
deleted file mode 100644
diff --git a/Engine/Plugins/Experimental/Animation/MotionTrajectory/Source/MotionTrajectory/Private/CharacterTrajectoryComponent.cpp b/Engine/Plugins/Experimental/Animation/MotionTrajectory/Source/MotionTrajectory/Private/CharacterTrajectoryComponent.cpp
--- a/Engine/Plugins/Experimental/Animation/MotionTrajectory/Source/MotionTrajectory/Private/CharacterTrajectoryComponent.cpp
+++ b/Engine/Plugins/Experimental/Animation/MotionTrajectory/Source/MotionTrajectory/Private/CharacterTrajectoryComponent.cpp
@@ -2,6 +2,7 @@
 
 #include "CharacterTrajectoryComponent.h"
 
+#include "Components/SkeletalMeshComponent.h"
 #include "GameFramework/Character.h"
 #include "HAL/IConsoleManager.h"
 #include "MotionTrajectory.h"
@@ -52,9 +53,26 @@ void UCharacterTrajectoryComponent::BeginPlay()
 	Super::BeginPlay();
 
 	SamplingData.Init();
-	CharacterTrajectoryData.Init(GetOwner());
+	TranslationHistory.Init(FVector::ZeroVector, SamplingData.NumHistorySamples);
 
-	FMotionTrajectoryLibrary::InitTrajectorySamples(Trajectory, CharacterTrajectoryData, SamplingData);
+	if (const ACharacter* Character = Cast<ACharacter>(GetOwner()))
+	{
+		if (const USkeletalMeshComponent* MeshComp = Character->GetMesh())
+		{
+			const FVector Position = MeshComp->GetComponentLocation();
+			const FQuat Facing = MeshComp->GetComponentRotation().Quaternion();
+
+			FMotionTrajectoryLibrary::InitTrajectorySamples(Trajectory, SamplingData, Position, Facing);
+		}
+		else
+		{
+			ensure(false);
+		}
+	}
+	else
+	{
+		UE_LOG(LogMotionTrajectory, Error, TEXT("UCharacterTrajectoryComponent requires its owner to be ACharacter"));
+	}
 }
 
 void UCharacterTrajectoryComponent::OnMovementUpdated(float DeltaSeconds, FVector OldLocation, FVector OldVelocity)
@@ -69,11 +87,15 @@ void UCharacterTrajectoryComponent::OnMovementUpdated(float DeltaSeconds, FVecto
 		return;
 	}
 
-	CharacterTrajectoryData.Update(DeltaSeconds);
-
-	FMotionTrajectoryLibrary::UpdateHistory_ShiftInWorldSpace(Trajectory, SamplingData, DeltaSeconds);
+	ACharacter* Character = Cast<ACharacter>(GetOwner());
+	if (!ensure(Character))
+	{
+		return;
+	}
 
-	FMotionTrajectoryLibrary::UpdatePrediction_SimulateCharacterMovement(Trajectory, CharacterTrajectoryData, SamplingData, DeltaSeconds);
+	CharacterTrajectoryData.UpdateDataFromCharacter(DeltaSeconds, Character);
+	FMotionTrajectoryLibrary::UpdateHistory_TransformHistory(Trajectory, TranslationHistory, CharacterTrajectoryData, SamplingData, DeltaSeconds);
+	FMotionTrajectoryLibrary::UpdatePrediction_SimulateCharacterMovement(Trajectory, CharacterTrajectoryData, SamplingData);
 
 	LastUpdateFrameNumber = GFrameNumber;
 
diff --git a/Engine/Plugins/Experimental/Animation/MotionTrajectory/Source/MotionTrajectory/Private/MotionTrajectoryLibrary.cpp b/Engine/Plugins/Experimental/Animation/MotionTrajectory/Source/MotionTrajectory/Private/MotionTrajectoryLibrary.cpp
--- a/Engine/Plugins/Experimental/Animation/MotionTrajectory/Source/MotionTrajectory/Private/MotionTrajectoryLibrary.cpp
+++ b/Engine/Plugins/Experimental/Animation/MotionTrajectory/Source/MotionTrajectory/Private/MotionTrajectoryLibrary.cpp
@@ -1,7 +1,6 @@
 // Copyright Epic Games, Inc. All Rights Reserved.
 
 #include "MotionTrajectoryLibrary.h"
-#include "CharacterMovementTrajectoryLibrary.h"
 #include "Components/SkeletalMeshComponent.h"
 #include "GameFramework/Character.h"
 #include "GameFramework/CharacterMovementComponent.h"
@@ -19,148 +18,169 @@ void FTrajectorySamplingData::Init()
 	SecondsPerPredictionSample = 1.f / PredictionSamplesPerSecond;
 }
 
-void FCharacterTrajectoryData::Init(const AActor* Actor)
+void FCharacterTrajectoryData::UpdateDataFromCharacter(float DeltaSeconds, const ACharacter* Character)
 {
-	Character = Cast<ACharacter>(Actor);
-	if (!ensureMsgf(Character, TEXT("FCharacterTrajectoryData requires valid ACharacter owner.")))
+	if (!ensure(Character))
 	{
 		return;
 	}
 
-	SkelMeshComponent = Character->GetMesh();
-	if (!ensureMsgf(SkelMeshComponent, TEXT("FCharacterTrajectoryData must be run on an ACharacter with a valid USkeletalMeshComponent.")))
+	// An AnimInstance might call this during an AnimBP recompile with 0 delta time.
+	if (DeltaSeconds <= 0.f)
 	{
 		return;
 	}
 
-	CharacterMovementComponent = Character->GetCharacterMovement();
-	if (!ensureMsgf(CharacterMovementComponent, TEXT("FCharacterTrajectoryData must be run on an ACharacter with a valid UCharacterMovementComponent.")))
+	if (const UCharacterMovementComponent* MoveComp = Character->GetCharacterMovement())
 	{
-		return;
-	}
-}
-
-void FCharacterTrajectoryData::Update(float DeltaSeconds)
-{
-	UpdateControllerRotationRate(DeltaSeconds);
-}
+		MaxSpeed = FMath::Max(MoveComp->GetMaxSpeed() * MoveComp->GetAnalogInputModifier(), MoveComp->GetMinAnalogSpeed());
+		BrakingDeceleration = FMath::Max(0.f, MoveComp->GetMaxBrakingDeceleration());
+		bOrientRotationToMovement = MoveComp->bOrientRotationToMovement;
 
-bool FCharacterTrajectoryData::IsValid() const
-{
-	// Init() won't initialize CharacterMovementComponent if Character or SkelMeshComponent are null.
-	check(!CharacterMovementComponent || (Character && SkelMeshComponent));
-
-	return (CharacterMovementComponent != nullptr);
-}
-
-void FCharacterTrajectoryData::UpdateControllerRotationRate(float DeltaSeconds)
-{
-	ControllerRotationRate = FRotator::ZeroRotator;
-	ControllerRotationRateClamped = FRotator::ZeroRotator;
+		Velocity = MoveComp->Velocity;
+		Acceleration = MoveComp->GetCurrentAcceleration();
 
-	if (!ensure(DeltaSeconds > 0.f))
-	{
-		return;
+		if (Acceleration.IsZero())
+		{
+			Friction = MoveComp->bUseSeparateBrakingFriction ? MoveComp->BrakingFriction : MoveComp->GroundFriction;
+			const float FrictionFactor = FMath::Max(0.f, MoveComp->BrakingFrictionFactor);
+			Friction = FMath::Max(0.f, Friction * FrictionFactor);
+		}
+		else
+		{
+			Friction = MoveComp->GroundFriction;
+		}
 	}
-
-	if (!ensure(IsValid()))
+	else
 	{
-		return;
+		ensure(false);
 	}
 
+	// @todo: Simulated proxies don't have controllers, so they'll need some other mechanism to account for controller rotation rate.
 	const AController* Controller = Character->Controller;
-	if (!Controller)
+	if (Controller)
 	{
-		// @todo: Simulated proxies don't have controllers, so they'll need some other mechanism to account for controller rotation rate.
-		return;
+		float DesiredControllerYaw = Controller->GetDesiredRotation().Yaw;
+		
+		const float DesiredYawDelta = DesiredControllerYaw - DesiredControllerYawLastUpdate;
+		DesiredControllerYawLastUpdate = DesiredControllerYaw;
+
+		ControllerYawRate = FRotator::NormalizeAxis(DesiredYawDelta) * (1.f / DeltaSeconds);
+		ControllerYawRateClamped = ControllerYawRate;
+		if (MaxControllerYawRate >= 0.f)
+		{
+			ControllerYawRateClamped = FMath::Sign(ControllerYawRate) * FMath::Min(FMath::Abs(ControllerYawRate), MaxControllerYawRate);
+		}
 	}
 
-	FRotator DesiredControllerRotation = Controller->GetDesiredRotation();
-	if (CharacterMovementComponent->ShouldRemainVertical())
+	if (const USkeletalMeshComponent* MeshComp = Character->GetMesh())
+	{
+		Position = MeshComp->GetComponentLocation();
+		Facing = MeshComp->GetComponentRotation().Quaternion();
+		MeshCompRelativeRotation = MeshComp->GetRelativeRotation().Quaternion();
+	}
+	else
 	{
-		DesiredControllerRotation.Yaw = FRotator::NormalizeAxis(DesiredControllerRotation.Yaw);
-		DesiredControllerRotation.Pitch = 0.f;
-		DesiredControllerRotation.Roll = 0.f;
+		ensure(false);
 	}
+}
 
-	const FRotator DesiredRotationDelta = DesiredControllerRotation - DesiredControllerRotationLastUpdate;
-	DesiredControllerRotationLastUpdate = DesiredControllerRotation;
+FVector FCharacterTrajectoryData::StepCharacterMovementGroundPrediction(float DeltaSeconds, const FVector& InVelocity, const FVector& InAcceleration) const
+{
+	FVector OutVelocity = InVelocity;
 
-	ControllerRotationRate = DesiredRotationDelta.GetNormalized() * (1.f / DeltaSeconds);
-	if (MaxControllerRotationRate >= 0.f)
+	// Braking logic is copied from UCharacterMovementComponent::ApplyVelocityBraking()
+	if (InAcceleration.IsZero())
 	{
-		ControllerRotationRateClamped.Pitch = FMath::Sign(ControllerRotationRate.Pitch) * FMath::Min(FMath::Abs(ControllerRotationRate.Pitch), MaxControllerRotationRate);
-		ControllerRotationRateClamped.Yaw = FMath::Sign(ControllerRotationRate.Yaw) * FMath::Min(FMath::Abs(ControllerRotationRate.Yaw), MaxControllerRotationRate);
-		ControllerRotationRateClamped.Roll = FMath::Sign(ControllerRotationRate.Roll) * FMath::Min(FMath::Abs(ControllerRotationRate.Roll), MaxControllerRotationRate);
+		if (InVelocity.IsZero())
+		{
+			return FVector::ZeroVector;
+		}
+
+		const bool bZeroFriction = (Friction == 0.f);
+		const bool bZeroBraking = (BrakingDeceleration == 0.f);
+
+		if (bZeroFriction && bZeroBraking)
+		{
+			return InVelocity;
+		}
+
+		static const float MaxTimeStep = 1.f / 60.f;
+		float RemainingTime = DeltaSeconds;
+
+		const FVector PrevLinearVelocity = OutVelocity;
+		const FVector RevAccel = (bZeroBraking ? FVector::ZeroVector : (-BrakingDeceleration * OutVelocity.GetSafeNormal()));
+
+		// Decelerate to brake to a stop
+		while (RemainingTime >= UCharacterMovementComponent::MIN_TICK_TIME)
+		{
+			// Zero friction uses constant deceleration, so no need for iteration.
+			const float dt = ((RemainingTime > MaxTimeStep && !bZeroFriction) ? FMath::Min(MaxTimeStep, RemainingTime * 0.5f) : RemainingTime);
+			RemainingTime -= dt;
+
+			// apply friction and braking
+			OutVelocity = OutVelocity + ((-Friction) * OutVelocity + RevAccel) * dt;
+
+			// Don't reverse direction
+			if ((OutVelocity | PrevLinearVelocity) <= 0.f)
+			{
+				OutVelocity = FVector::ZeroVector;
+				return OutVelocity;
+			}
+		}
+
+		// Clamp to zero if nearly zero, or if below min threshold and braking
+		const float VSizeSq = OutVelocity.SizeSquared();
+		if (VSizeSq <= KINDA_SMALL_NUMBER || (!bZeroBraking && VSizeSq <= FMath::Square(UCharacterMovementComponent::BRAKE_TO_STOP_VELOCITY)))
+		{
+			OutVelocity = FVector::ZeroVector;
+		}
 	}
+	// Acceleration logic is copied from  UCharacterMovementComponent::CalcVelocity
 	else
 	{
-		ControllerRotationRateClamped = ControllerRotationRate;
+		const FVector AccelDir = InAcceleration.GetSafeNormal();
+		const float VelSize = OutVelocity.Size();
+
+		OutVelocity = OutVelocity - (OutVelocity - AccelDir * VelSize) * FMath::Min(DeltaSeconds * Friction, 1.f);
+
+		OutVelocity += InAcceleration * DeltaSeconds;
+		OutVelocity = OutVelocity.GetClampedToMaxSize(MaxSpeed);
 	}
+
+	return OutVelocity;
 }
 
 void FMotionTrajectoryLibrary::InitTrajectorySamples(FPoseSearchQueryTrajectory& Trajectory,
-	const FCharacterTrajectoryData& CharacterTrajectoryData, const FTrajectorySamplingData& SamplingData)
+	const FTrajectorySamplingData& SamplingData, const FVector& Position, const FQuat& Facing)
 {
-	if (!ensure(CharacterTrajectoryData.SkelMeshComponent))
-	{
-		return;
-	}
-
-	const FVector PositionWS = CharacterTrajectoryData.SkelMeshComponent->GetComponentLocation();;
-	const FQuat FacingWS = CharacterTrajectoryData.SkelMeshComponent->GetComponentRotation().Quaternion();
-
 	// History + current sample + prediction
 	Trajectory.Samples.SetNumUninitialized(SamplingData.NumHistorySamples + 1 + SamplingData.NumPredictionSamples);
 
 	// Initialize history samples
 	for (int32 i = 0; i < SamplingData.NumHistorySamples; ++i)
 	{
-		Trajectory.Samples[i].Position = PositionWS;
-		Trajectory.Samples[i].Facing = FacingWS;
+		Trajectory.Samples[i].Position = Position;
+		Trajectory.Samples[i].Facing = Facing;
 		Trajectory.Samples[i].AccumulatedSeconds = SamplingData.SecondsPerHistorySample * (i - SamplingData.NumHistorySamples);
 	}
 
 	// Initialize current sample and prediction
 	for (int32 i = SamplingData.NumHistorySamples; i < Trajectory.Samples.Num(); ++i)
 	{
-		Trajectory.Samples[i].Position = PositionWS;
-		Trajectory.Samples[i].Facing = FacingWS;
+		Trajectory.Samples[i].Position = Position;
+		Trajectory.Samples[i].Facing = Facing;
 		Trajectory.Samples[i].AccumulatedSeconds = SamplingData.SecondsPerPredictionSample * (i - SamplingData.NumHistorySamples);
 	}
 }
 
-void FMotionTrajectoryLibrary::UpdateHistory_ShiftInWorldSpace(FPoseSearchQueryTrajectory& Trajectory,
-	const FTrajectorySamplingData& SamplingData, float DeltaSeconds)
-{
-	check(SamplingData.NumHistorySamples <= Trajectory.Samples.Num());
-
-	// Shift history Samples when it's time to record a new one.
-	if (SamplingData.NumHistorySamples > 0 && FMath::Abs(Trajectory.Samples[SamplingData.NumHistorySamples - 1].AccumulatedSeconds) >= SamplingData.SecondsPerHistorySample)
-	{
-		for (int32 Index = 0; Index < SamplingData.NumHistorySamples; ++Index)
-		{
-			Trajectory.Samples[Index] = Trajectory.Samples[Index + 1];
-			Trajectory.Samples[Index].AccumulatedSeconds -= DeltaSeconds;
-		}
-	}
-	else
-	{
-		for (int32 Index = 0; Index < SamplingData.NumHistorySamples; ++Index)
-		{
-			Trajectory.Samples[Index].AccumulatedSeconds -= DeltaSeconds;
-		}
-	}
-}
-
 void FMotionTrajectoryLibrary::UpdateHistory_TransformHistory(FPoseSearchQueryTrajectory& Trajectory, TArrayView<FVector> TranslationHistory,
 	const FCharacterTrajectoryData& CharacterTrajectoryData, const FTrajectorySamplingData& SamplingData, float DeltaSeconds)
 {
 	check(SamplingData.NumHistorySamples <= Trajectory.Samples.Num());
 	check(TranslationHistory.Num() == SamplingData.NumHistorySamples);
 
-	const FVector& CurrentVelocityWS = CharacterTrajectoryData.CharacterMovementComponent->Velocity;
-	FVector CurrentTranslation = CurrentVelocityWS * DeltaSeconds;
+	FVector CurrentTranslation = CharacterTrajectoryData.Velocity * DeltaSeconds;
 
 	// Shift history Samples when it's time to record a new one.
 	if (SamplingData.NumHistorySamples > 0 && FMath::Abs(Trajectory.Samples[SamplingData.NumHistorySamples - 1].AccumulatedSeconds) >= SamplingData.SecondsPerHistorySample)
@@ -184,11 +204,10 @@ void FMotionTrajectoryLibrary::UpdateHistory_TransformHistory(FPoseSearchQueryTr
 	}
 
 	// Update trajectory samples by applying the tracked translations to the current world position.
-	FVector CurrentPositionWS = CharacterTrajectoryData.SkelMeshComponent->GetComponentLocation();
 	for (int32 Index = 0; Index < SamplingData.NumHistorySamples; ++Index)
 	{
 		Trajectory.Samples[Index].AccumulatedSeconds -= DeltaSeconds;
-		Trajectory.Samples[Index].Position = CurrentPositionWS - TranslationHistory[Index];
+		Trajectory.Samples[Index].Position = CharacterTrajectoryData.Position - TranslationHistory[Index];
 
 		// @todo: Handle facing. We currently don't use facing for history in any of our content. We will need the rotation intent of the character.
 		Trajectory.Samples[Index].Facing = FQuat::Identity;
@@ -211,28 +230,23 @@ FVector FMotionTrajectoryLibrary::RemapVectorMagnitudeWithCurve(const FVector& V
 }
 
 void FMotionTrajectoryLibrary::UpdatePrediction_SimulateCharacterMovement(FPoseSearchQueryTrajectory& Trajectory,
-	const FCharacterTrajectoryData& CharacterTrajectoryData, const FTrajectorySamplingData& SamplingData, float DeltaSeconds)
+	const FCharacterTrajectoryData& CharacterTrajectoryData, const FTrajectorySamplingData& SamplingData)
 {
-	if (!ensure(CharacterTrajectoryData.IsValid()))
-	{
-		return;
-	}
-
-	FVector CurrentPositionWS = CharacterTrajectoryData.SkelMeshComponent->GetComponentLocation();
-	FVector CurrentVelocityWS = FMotionTrajectoryLibrary::RemapVectorMagnitudeWithCurve(CharacterTrajectoryData.CharacterMovementComponent->Velocity,
+	FVector CurrentPositionWS = CharacterTrajectoryData.Position;
+	FVector CurrentVelocityWS = FMotionTrajectoryLibrary::RemapVectorMagnitudeWithCurve(CharacterTrajectoryData.Velocity,
 		CharacterTrajectoryData.bUseSpeedRemappingCurve, CharacterTrajectoryData.SpeedRemappingCurve);
-	FVector CurrentAccelerationWS = FMotionTrajectoryLibrary::RemapVectorMagnitudeWithCurve(CharacterTrajectoryData.CharacterMovementComponent->GetCurrentAcceleration(),
+	FVector CurrentAccelerationWS = FMotionTrajectoryLibrary::RemapVectorMagnitudeWithCurve(CharacterTrajectoryData.Acceleration,
 		CharacterTrajectoryData.bUseAccelerationRemappingCurve, CharacterTrajectoryData.AccelerationRemappingCurve);
 
 	// bending CurrentVelocityWS towards CurrentAccelerationWS
 	if (CharacterTrajectoryData.BendVelocityTowardsAcceleration > UE_KINDA_SMALL_NUMBER && !CurrentAccelerationWS.IsNearlyZero())
 	{
 		const float CurrentSpeed = CurrentVelocityWS.Length();
-		const FVector VelocityWSAlongAcceleration = CurrentAccelerationWS.GetUnsafeNormal()* CurrentSpeed;
+		const FVector VelocityWSAlongAcceleration = CurrentAccelerationWS.GetUnsafeNormal() * CurrentSpeed;
 		if (CharacterTrajectoryData.BendVelocityTowardsAcceleration < 1.f - UE_KINDA_SMALL_NUMBER)
 		{
 			CurrentVelocityWS = FMath::Lerp(CurrentVelocityWS, VelocityWSAlongAcceleration, CharacterTrajectoryData.BendVelocityTowardsAcceleration);
-			
+
 			const float NewLength = CurrentVelocityWS.Length();
 			if (NewLength > UE_KINDA_SMALL_NUMBER)
 			{
@@ -249,10 +263,10 @@ void FMotionTrajectoryLibrary::UpdatePrediction_SimulateCharacterMovement(FPoseS
 		}
 	}
 
-	FQuat CurrentFacingWS = CharacterTrajectoryData.SkelMeshComponent->GetComponentRotation().Quaternion();
-	FQuat SkelMeshCompRelativeRotation = CharacterTrajectoryData.SkelMeshComponent->GetRelativeRotation().Quaternion();
+	FQuat CurrentFacingWS = CharacterTrajectoryData.Facing;
+	FQuat SkelMeshCompRelativeRotation = CharacterTrajectoryData.MeshCompRelativeRotation;
 
-	FQuat ControllerRotationPerStep = (CharacterTrajectoryData.ControllerRotationRateClamped * SamplingData.SecondsPerPredictionSample).Quaternion();
+	FQuat ControllerRotationPerStep = FQuat::MakeFromEuler(FVector(0.f, 0.f, CharacterTrajectoryData.ControllerYawRateClamped * SamplingData.SecondsPerPredictionSample));
 
 	float AccumulatedSeconds = 0.f;
 
@@ -278,13 +292,12 @@ void FMotionTrajectoryLibrary::UpdatePrediction_SimulateCharacterMovement(FPoseS
 			CurrentAccelerationWS = FMotionTrajectoryLibrary::RemapVectorMagnitudeWithCurve(ControllerRotationPerStep * CurrentAccelerationWS,
 				CharacterTrajectoryData.bUseAccelerationRemappingCurve, CharacterTrajectoryData.AccelerationRemappingCurve);
 
-			FVector NewVelocityCS = FVector::ZeroVector;
-			UCharacterMovementTrajectoryLibrary::StepCharacterMovementGroundPrediction(SamplingData.SecondsPerPredictionSample, CurrentVelocityWS, CurrentAccelerationWS, 
-				CharacterTrajectoryData.CharacterMovementComponent, NewVelocityCS);
-			CurrentVelocityWS = FMotionTrajectoryLibrary::RemapVectorMagnitudeWithCurve(NewVelocityCS,
+			FVector NewVelocityWS = CharacterTrajectoryData.StepCharacterMovementGroundPrediction(SamplingData.SecondsPerPredictionSample, CurrentVelocityWS, CurrentAccelerationWS);
+
+			CurrentVelocityWS = FMotionTrajectoryLibrary::RemapVectorMagnitudeWithCurve(NewVelocityWS,
 				CharacterTrajectoryData.bUseSpeedRemappingCurve, CharacterTrajectoryData.SpeedRemappingCurve);
 
-			if (CharacterTrajectoryData.CharacterMovementComponent->bOrientRotationToMovement && !CurrentAccelerationWS.IsNearlyZero())
+			if (CharacterTrajectoryData.bOrientRotationToMovement && !CurrentAccelerationWS.IsNearlyZero())
 			{
 				// Rotate towards acceleration.
 				const FVector CurrentAccelerationCS = SkelMeshCompRelativeRotation.RotateVector(CurrentAccelerationWS);
diff --git a/Engine/Plugins/Experimental/Animation/MotionTrajectory/Source/MotionTrajectory/Public/CharacterMovementTrajectoryLibrary.h b/Engine/Plugins/Experimental/Animation/MotionTrajectory/Source/MotionTrajectory/Public/CharacterMovementTrajectoryLibrary.h
deleted file mode 100644
diff --git a/Engine/Plugins/Experimental/Animation/MotionTrajectory/Source/MotionTrajectory/Public/CharacterTrajectoryComponent.h b/Engine/Plugins/Experimental/Animation/MotionTrajectory/Source/MotionTrajectory/Public/CharacterTrajectoryComponent.h
--- a/Engine/Plugins/Experimental/Animation/MotionTrajectory/Source/MotionTrajectory/Public/CharacterTrajectoryComponent.h
+++ b/Engine/Plugins/Experimental/Animation/MotionTrajectory/Source/MotionTrajectory/Public/CharacterTrajectoryComponent.h
@@ -39,5 +39,7 @@ protected:
 	UPROPERTY(Transient, BlueprintReadOnly, Category = "Trajectory Settings")
 	FCharacterTrajectoryData CharacterTrajectoryData;
 
+	TArray<FVector> TranslationHistory;
+
 	uint32 LastUpdateFrameNumber = 0;
 };
\ No newline at end of file
diff --git a/Engine/Plugins/Experimental/Animation/MotionTrajectory/Source/MotionTrajectory/Public/MotionTrajectoryLibrary.h b/Engine/Plugins/Experimental/Animation/MotionTrajectory/Source/MotionTrajectory/Public/MotionTrajectoryLibrary.h
--- a/Engine/Plugins/Experimental/Animation/MotionTrajectory/Source/MotionTrajectory/Public/MotionTrajectoryLibrary.h
+++ b/Engine/Plugins/Experimental/Animation/MotionTrajectory/Source/MotionTrajectory/Public/MotionTrajectoryLibrary.h
@@ -47,10 +47,9 @@ struct MOTIONTRAJECTORY_API FCharacterTrajectoryData
 public:
 	GENERATED_BODY()
 
-	void Init(const AActor* Actor);
-	void Update(float DeltaSeconds);
+	void UpdateDataFromCharacter(float DeltaSeconds, const ACharacter* Character);
 
-	bool IsValid() const;
+	FVector StepCharacterMovementGroundPrediction(float DeltaSeconds, const FVector& InVelocity, const FVector& InAcceleration) const;
 
 	// If the character is forward facing (i.e. bOrientRotationToMovement is true), this controls how quickly the trajectory will rotate
 	// to face acceleration. It's common for this to differ from the rotation rate of the character, because animations are often authored 
@@ -58,10 +57,10 @@ public:
 	UPROPERTY(EditDefaultsOnly, Category = "Trajectory Settings")
 	float RotateTowardsMovementSpeed = 10.f;
 
-	// Maximum controller rotation rate in degrees per second used to clamp the character owner controller desired rotation to generate the prediction trajectory.
+	// Maximum controller yaw  rate in degrees per second used to clamp the character owner controller desired yaw to generate the prediction trajectory.
 	// Negative values disable the clamping behavior
 	UPROPERTY(EditDefaultsOnly, Category = "Trajectory Settings")
-	float MaxControllerRotationRate = -1.f;
+	float MaxControllerYawRate = 70.f;
 
 	// artificially bend character velocity towards acceleration direction to compute trajectory prediction, to get sharper turns
 	// 0: character velocity is used with no alteration, 1: the acceleration direction is used as velocity direction
@@ -80,22 +79,21 @@ public:
 	UPROPERTY(EditDefaultsOnly, Category = "Trajectory Settings", meta = (EditCondition = "bUseAccelerationRemappingCurve"))
 	FRuntimeFloatCurve AccelerationRemappingCurve;
 
-	UPROPERTY(Transient)
-	TObjectPtr<const ACharacter> Character = nullptr;
+	float ControllerYawRate = 0.f;
+	float ControllerYawRateClamped = 0.f;
+	float DesiredControllerYawLastUpdate = 0.f;
 
-	UPROPERTY(Transient)
-	TObjectPtr<const USkeletalMeshComponent> SkelMeshComponent = nullptr;
+	float MaxSpeed = 0.f;
+	float BrakingDeceleration = 0.f;
+	float Friction = 0.f;
 
-	UPROPERTY(Transient)
-	TObjectPtr<const UCharacterMovementComponent> CharacterMovementComponent = nullptr;
+	FVector Velocity = FVector::ZeroVector;
+	FVector Acceleration = FVector::ZeroVector;
 
-	FRotator ControllerRotationRate = FRotator::ZeroRotator;
-	FRotator ControllerRotationRateClamped = FRotator::ZeroRotator;
-
-private:
-	void UpdateControllerRotationRate(float DeltaSeconds);
-
-	FRotator DesiredControllerRotationLastUpdate = FRotator::ZeroRotator;
+	FVector Position = FVector::ZeroVector;
+	FQuat Facing = FQuat::Identity;
+	FQuat MeshCompRelativeRotation = FQuat::Identity;
+	bool bOrientRotationToMovement = false;
 };
 
 /**
@@ -106,22 +104,17 @@ struct MOTIONTRAJECTORY_API FMotionTrajectoryLibrary
 {
 public:
 	static void InitTrajectorySamples(FPoseSearchQueryTrajectory& Trajectory,
-		const FCharacterTrajectoryData& CharacterTrajectoryData, const FTrajectorySamplingData& SamplingData);
-
-	// Update history by shifting world space samples back in time. This gives us a history of world space transforms.
-	// This approach is simple, but it causes problems with moving platforms.
-	static void UpdateHistory_ShiftInWorldSpace(FPoseSearchQueryTrajectory& Trajectory,
-		const FTrajectorySamplingData& SamplingData, float DeltaSeconds);
+		const FTrajectorySamplingData& SamplingData, const FVector& Position, const FQuat& Facing);
 
 	// Update history by tracking offsets that result from character intent (e.g. movement component velocity) and applying
 	// that to the current world transform. This works well on moving platforms as it only stores a history of movement
-	// that results from character intent.
+	// that results from character intent, not movement from platforms.
 	static void UpdateHistory_TransformHistory(FPoseSearchQueryTrajectory& Trajectory, TArrayView<FVector> TranslationHistory,
 		const FCharacterTrajectoryData& CharacterTrajectoryData, const FTrajectorySamplingData& SamplingData, float DeltaSeconds);
 
-	// Update prediction by simulating the movement math for ground locomotion in UCharacterMovementComponent.
+	// Update prediction by simulating the movement math for ground locomotion from UCharacterMovementComponent.
 	static void UpdatePrediction_SimulateCharacterMovement(FPoseSearchQueryTrajectory& Trajectory,
-		const FCharacterTrajectoryData& CharacterTrajectoryData, const FTrajectorySamplingData& SamplingData, float DeltaSeconds);
+		const FCharacterTrajectoryData& CharacterTrajectoryData, const FTrajectorySamplingData& SamplingData);
 
 private:
 	static FVector RemapVectorMagnitudeWithCurve(const FVector& Vector, bool bUseCurve, const FRuntimeFloatCurve& Curve);
